SparseVolumeTexture asset now stores the source file and cooks using that source file instead of data generated during the import.

#rnx

diff --git a/Engine/Source/Editor/SparseVolumeTexture/Private/SparseVolumeTextureFactory.cpp b/Engine/Source/Editor/SparseVolumeTexture/Private/SparseVolumeTextureFactory.cpp
--- a/Engine/Source/Editor/SparseVolumeTexture/Private/SparseVolumeTextureFactory.cpp
+++ b/Engine/Source/Editor/SparseVolumeTexture/Private/SparseVolumeTextureFactory.cpp
@@ -7,10 +7,12 @@
 #if WITH_EDITOR
 
 #include "SparseVolumeTextureOpenVDB.h"
+#include "SparseVolumeTextureOpenVDBUtility.h"
 
 #include "Serialization/EditorBulkDataWriter.h"
 #include "Misc/Paths.h"
 #include "Misc/ScopedSlowTask.h"
+#include "Misc/FileHelper.h"
 
 #include "Editor.h"
 
@@ -93,239 +95,6 @@ void USparseVolumeTextureFactory::CleanUp()
 	Super::CleanUp();
 }
 
-struct FOpenVDBData
-{
-	FVector VolumeActiveAABBMin;
-	FVector VolumeActiveAABBMax;
-	FVector VolumeActiveDim;
-	FVector VolumeVoxelSize;
-	bool bIsInWorldSpace;
-	bool bHasUniformVoxels;
-};
-
-#if OPENVDB_AVAILABLE
-
-static FOpenVDBData GetOpenVDBData(openvdb::GridBase::Ptr GridBase)
-{
-	FOpenVDBData OpenVDBData;
-	openvdb::CoordBBox VolumeActiveAABB = GridBase->evalActiveVoxelBoundingBox();
-	openvdb::Coord VolumeActiveDim = GridBase->evalActiveVoxelDim();
-	openvdb::Vec3d VolumeVoxelSize = GridBase->voxelSize();
-
-	OpenVDBData.VolumeActiveAABBMin = FVector(VolumeActiveAABB.min().x(), VolumeActiveAABB.min().y(), VolumeActiveAABB.min().z());
-	OpenVDBData.VolumeActiveAABBMax = FVector(VolumeActiveAABB.max().x(), VolumeActiveAABB.max().y(), VolumeActiveAABB.max().z());
-	OpenVDBData.VolumeActiveDim = FVector(VolumeActiveDim.x(), VolumeActiveDim.y(), VolumeActiveDim.z());
-	OpenVDBData.VolumeVoxelSize = FVector(VolumeVoxelSize.x(), VolumeVoxelSize.y(), VolumeVoxelSize.z());
-	OpenVDBData.bIsInWorldSpace = GridBase->isInWorldSpace();
-	OpenVDBData.bHasUniformVoxels = GridBase->hasUniformVoxels();
-
-	return OpenVDBData;
-}
-
-static bool IsOpenVDBDataValid(FOpenVDBData& OpenVDBData, const FString& Filename)
-{
-	if (OpenVDBData.VolumeActiveDim.X * OpenVDBData.VolumeActiveDim.Y * OpenVDBData.VolumeActiveDim.Z == 0)
-	{
-		// SVT_TODO we should gently handle that case
-		UE_LOG(LogSparseVolumeTextureFactory, Error, TEXT("Could not import empty OpenVDB asset due to volume size being 0: %s"), *Filename);
-		return false;
-	}
-
-	if (!OpenVDBData.bHasUniformVoxels)
-	{
-		UE_LOG(LogSparseVolumeTextureFactory, Error, TEXT("OpenVDB importer cannot handle non uniform voxels: %s"), *Filename);
-		return false;
-	}
-	return true;
-}
-
-
-static bool ConvertOpenVDBToSparseVolumeTexture(
-	openvdb::FloatGrid::Ptr DensityGrid,
-	const FString& SourceFilename,
-	FSparseVolumeRawSource& OutSparseVolumeRawSource,
-	bool bOverrideActiveMinMax = false,
-	FVector ActiveMin = FVector::ZeroVector,
-	FVector ActiveMax = FVector::ZeroVector)
-{
-	FOpenVDBData OVDBData = GetOpenVDBData(DensityGrid);
-	if (!IsOpenVDBDataValid(OVDBData, SourceFilename))
-	{
-		return false;
-	}
-
-	if (bOverrideActiveMinMax)
-	{
-		OVDBData.VolumeActiveAABBMin = ActiveMin;
-		OVDBData.VolumeActiveAABBMax = ActiveMax;
-		OVDBData.VolumeActiveDim = ActiveMax - ActiveMin;
-	}
-
-	FSparseVolumeAssetHeader& Header = OutSparseVolumeRawSource.Header;
-
-	Header.SourceVolumeResolution = FIntVector(OVDBData.VolumeActiveDim);
-	// SVT_TODO handle error too high volume resolution assuming dx11 limitations?
-
-	Header.PageTableVolumeResolution = FIntVector3(
-		FMath::DivideAndRoundUp(Header.SourceVolumeResolution.X, SPARSE_VOLUME_TILE_RES),
-		FMath::DivideAndRoundUp(Header.SourceVolumeResolution.Y, SPARSE_VOLUME_TILE_RES),
-		FMath::DivideAndRoundUp(Header.SourceVolumeResolution.Z, SPARSE_VOLUME_TILE_RES));
-	Header.TileDataVolumeResolution = FIntVector::ZeroValue;	// unknown for now
-
-	const uint32 VolumeTileDataBytes = SPARSE_VOLUME_TILE_RES * SPARSE_VOLUME_TILE_RES * SPARSE_VOLUME_TILE_RES;
-	uint32 NumAllocatedPages = 0;
-
-	// Allocate some memory for temp data (worst case)
-	TArray<uint8> LinearAllocatedTiles;
-	TArray<FIntVector3> LinearAllocatedPages;
-	LinearAllocatedTiles.SetNum(Header.PageTableVolumeResolution.X * Header.PageTableVolumeResolution.Y * Header.PageTableVolumeResolution.Z * VolumeTileDataBytes);
-	LinearAllocatedPages.SetNum(Header.PageTableVolumeResolution.X * Header.PageTableVolumeResolution.Y * Header.PageTableVolumeResolution.Z);
-
-	auto accessor = DensityGrid->getConstAccessor();
-	const float GridBackgroundValue = DensityGrid->background();
-
-	// Go over each potential page table form the source data and push allocate it if it has any data.
-	// Otherwise point to the default empty page.
-	bool bAnyEmptyPageExists = false;
-	for (int32_t PageZ = 0; PageZ < Header.PageTableVolumeResolution.Z; ++PageZ)
-	{
-		for (int32_t PageY = 0; PageY < Header.PageTableVolumeResolution.Y; ++PageY)
-		{
-			for (int32_t PageX = 0; PageX < Header.PageTableVolumeResolution.X; ++PageX)
-			{
-				bool bHasAnyData = false;
-
-				uint8* NewTileData = LinearAllocatedTiles.GetData() + NumAllocatedPages * VolumeTileDataBytes;
-
-				for (int32_t z = 0; z < SPARSE_VOLUME_TILE_RES; ++z)
-				{
-					for (int32_t y = 0; y < SPARSE_VOLUME_TILE_RES; ++y)
-					{
-						for (int32_t x = 0; x < SPARSE_VOLUME_TILE_RES; ++x)
-						{
-							FVector VoxelCoord(PageX * SPARSE_VOLUME_TILE_RES + x, PageY * SPARSE_VOLUME_TILE_RES + y, PageZ * SPARSE_VOLUME_TILE_RES + z);	// This assumes sampling outside the boundary retuirn a default value
-							FVector VoxelCoord2 = OVDBData.VolumeActiveAABBMin + VoxelCoord;
-							float VoxelValue = accessor.getValue(openvdb::Coord(VoxelCoord2.X, VoxelCoord2.Y, VoxelCoord2.Z));
-
-							// clamp for now until we have fp16
-							VoxelValue = FMath::Clamp(VoxelValue, 0.0f, 1.0f);
-
-							bHasAnyData |= VoxelValue > GridBackgroundValue;
-
-							NewTileData[z * SPARSE_VOLUME_TILE_RES * SPARSE_VOLUME_TILE_RES + y * SPARSE_VOLUME_TILE_RES + x] = uint8(VoxelValue * 255.0f);
-						}
-					}
-				}
-
-				if (bHasAnyData)
-				{
-					LinearAllocatedPages[NumAllocatedPages] = FIntVector3(PageX, PageY, PageZ);
-					NumAllocatedPages++;
-				}
-				bAnyEmptyPageExists |= !bHasAnyData;
-			}
-		}
-	}
-
-	// Compute Page and Tile VolumeResolution from allocated pages
-	const uint32 EffectivelyAllocatedPageEntries = NumAllocatedPages + (bAnyEmptyPageExists ? 1 : 0);
-	uint32 TileVolumeResolutionCube = 1;
-	while (TileVolumeResolutionCube * TileVolumeResolutionCube * TileVolumeResolutionCube < EffectivelyAllocatedPageEntries)
-	{
-		TileVolumeResolutionCube++;				// We use a simple loop to compute the minimum resolution of a cube to store all the tile data
-	}
-	Header.TileDataVolumeResolution = FIntVector3(TileVolumeResolutionCube, TileVolumeResolutionCube, TileVolumeResolutionCube);
-	while (Header.TileDataVolumeResolution.X * Header.TileDataVolumeResolution.Y * (Header.TileDataVolumeResolution.Z - 1) > int32(EffectivelyAllocatedPageEntries))
-	{
-		Header.TileDataVolumeResolution.Z--;	// We then trim an edge to get back space.
-	}
-	const FIntVector3 TileCoordResolution = Header.TileDataVolumeResolution;
-	Header.TileDataVolumeResolution = Header.TileDataVolumeResolution * SPARSE_VOLUME_TILE_RES;
-
-	// Initialise the SparseVolumeTexture page and tile.
-	// It would be good to stream in the OpenVDB file there but could not find a way to do that yet so we just store the sparse texture.
-	OutSparseVolumeRawSource.DensityPage.SetNumZeroed(Header.PageTableVolumeResolution.X * Header.PageTableVolumeResolution.Y * Header.PageTableVolumeResolution.Z);
-	OutSparseVolumeRawSource.DensityData.SetNumZeroed(Header.TileDataVolumeResolution.X * Header.TileDataVolumeResolution.Y * Header.TileDataVolumeResolution.Z);
-
-	FIntVector DestinationTileCoord = FIntVector::ZeroValue;
-	auto GoToNextTileCoord = [&]()
-	{
-		DestinationTileCoord.X++;
-		if (DestinationTileCoord.X >= TileCoordResolution.X)
-		{
-			DestinationTileCoord.X = 0;
-			DestinationTileCoord.Y++;
-		}
-		if (DestinationTileCoord.Y >= TileCoordResolution.Y)
-		{
-			DestinationTileCoord.Y = 0;
-			DestinationTileCoord.Z++;
-		}
-	};
-
-	// Add an empty tile is needed, reserve slot at coord 0
-	if (bAnyEmptyPageExists)
-	{
-		// DensityPage is all cleared to zero, simply skip a tile
-		GoToNextTileCoord();
-	}
-
-	// Generate page table and tile volume data by splatting the data
-	for (uint32 i = 0; i < NumAllocatedPages; ++i)
-	{
-		FIntVector3 PageCoordToSplat = LinearAllocatedPages[i];
-		uint8* TileDataToSplat = LinearAllocatedTiles.GetData() + i * VolumeTileDataBytes;
-
-		// A page encodes the physical tile coord as unsigned int of 11 11 10 bits
-		// This means a page coord cannot be larger than 2047 for x and y and 1023 for z
-		// which mean we cannot have more than 2048*2048*1024 = 4 Giga tiles of 16^3 tiles.
-		uint32 DestinationTileCoord32bit = (DestinationTileCoord.X & 0x7FF) | ((DestinationTileCoord.Y & 0x7FF) << 11) | ((DestinationTileCoord.Z & 0x3FF) << 22);
-
-		// Setup the page table entry
-		OutSparseVolumeRawSource.DensityPage.GetData()
-			[
-				PageCoordToSplat.Z * Header.PageTableVolumeResolution.X * Header.PageTableVolumeResolution.Y +
-				PageCoordToSplat.Y * Header.PageTableVolumeResolution.X +
-			PageCoordToSplat.X
-			] = DestinationTileCoord32bit;
-
-		// Now copy the tile data
-		for (int32_t z = 0; z < SPARSE_VOLUME_TILE_RES; ++z)
-		{
-			for (int32_t y = 0; y < SPARSE_VOLUME_TILE_RES; ++y)
-			{
-				for (int32_t x = 0; x < SPARSE_VOLUME_TILE_RES; ++x)
-				{
-					OutSparseVolumeRawSource.DensityData.GetData()
-						[
-							(DestinationTileCoord.Z * SPARSE_VOLUME_TILE_RES + z) * Header.TileDataVolumeResolution.X * Header.TileDataVolumeResolution.Y +
-							(DestinationTileCoord.Y * SPARSE_VOLUME_TILE_RES + y) * Header.TileDataVolumeResolution.X +
-						(DestinationTileCoord.X * SPARSE_VOLUME_TILE_RES + x)
-						] = TileDataToSplat[z * SPARSE_VOLUME_TILE_RES * SPARSE_VOLUME_TILE_RES + y * SPARSE_VOLUME_TILE_RES + x];
-				}
-			}
-		}
-
-		// Set the next tile to be written to
-		GoToNextTileCoord();
-	}
-
-	return true;
-}
-
-struct FOpenVDBFrameData
-{
-	FString Filename;
-	openvdb::io::File* File;
-
-	FOpenVDBData OpenVDBData;
-
-	openvdb::FloatGrid::Ptr DensityGrid;
-};
-
-#endif // OPENVDB_AVAILABLE
-
-
 UObject* USparseVolumeTextureFactory::FactoryCreateFile(UClass* InClass, UObject* InParent, FName InName, EObjectFlags Flags, const FString& Filename,
 	const TCHAR* Parms, FFeedbackContext* Warn, bool& bOutOperationCanceled)
 {
@@ -386,130 +155,55 @@ UObject* USparseVolumeTextureFactory::FactoryCreateFile(UClass* InClass, UObject
 		{
 			FrameCount++;
 		}
-		TArray<FOpenVDBFrameData> OpenVDBFramesData;
-		OpenVDBFramesData.SetNum(FrameCount);
 
 		FScopedSlowTask ImportTask(FrameCount, LOCTEXT("ImportingVDBAnim", "Importing OpenVDB animation"));
 		ImportTask.MakeDialog(true);
 
-		// Open all each frame openvdb files.
-		for(int32 FrameIndex = 0; FrameIndex < FrameCount; ++FrameIndex)
-		{
-			FOpenVDBFrameData& OpenVDBFrameData = OpenVDBFramesData[FrameIndex];
-			OpenVDBFrameData.Filename = GetOpenVDBFileNameForFrame(FrameIndex);
-
-			std::string FileNameStr(TCHAR_TO_ANSI(*OpenVDBFrameData.Filename));
-			OpenVDBFrameData.File = new openvdb::io::File(FileNameStr.c_str());
-			try
-			{
-				OpenVDBFrameData.File->open();
-			}
-			catch (openvdb::Exception e)
-			{
-				UE_LOG(LogSparseVolumeTextureFactory, Error, TEXT("OpenVDB file coult not be opened: %s"), *OpenVDBFrameData.Filename);
-				return nullptr;
-			}
-
-			openvdb::GridBase::Ptr BaseGrid = nullptr;
-			for (openvdb::io::File::NameIterator nameIter = OpenVDBFrameData.File->beginName();
-				nameIter != OpenVDBFrameData.File->endName(); ++nameIter)
-			{
-				if (nameIter.gridName() == "density")
-				{
-					openvdb::GridBase::Ptr QueryGrid = OpenVDBFrameData.File->readGrid(nameIter.gridName());
-					if (QueryGrid->isType<openvdb::FloatGrid>())
-					{
-						BaseGrid = QueryGrid;
-						break;
-					}
-					break;
-				}
-			}
-
-			// If we have not found any density map, let's pick up the first float map
-			for (openvdb::io::File::NameIterator nameIter = OpenVDBFrameData.File->beginName();
-				nameIter != OpenVDBFrameData.File->endName(); ++nameIter)
-			{
-				openvdb::GridBase::Ptr QueryGrid = OpenVDBFrameData.File->readGrid(nameIter.gridName());
-				if (QueryGrid->isType<openvdb::FloatGrid>())
-				{
-					BaseGrid = QueryGrid;
-					break;
-				}
-			}
-
-			OpenVDBFrameData.DensityGrid = openvdb::gridPtrCast<openvdb::FloatGrid>(BaseGrid);
-
-			if (OpenVDBFrameData.DensityGrid == nullptr)
-			{
-				UE_LOG(LogSparseVolumeTextureFactory, Error, TEXT("Could not find a Grid named density or grid of type float: %s"), *OpenVDBFrameData.Filename);
-				return nullptr;
-			}
-
-			OpenVDBFrameData.OpenVDBData = GetOpenVDBData(OpenVDBFrameData.DensityGrid);
-			if (!IsOpenVDBDataValid(OpenVDBFrameData.OpenVDBData, OpenVDBFrameData.Filename))
-			{
-				return nullptr;
-			}
-		}
-
-
-		// Evaluate the maximum extent for all frames.
-		FVector AnimationAABBMin;
-		FVector AnimationAABBMax;
-		for (int32 FrameIndex = 0; FrameIndex < FrameCount; ++FrameIndex)
-		{
-			FOpenVDBFrameData& OpenVDBFrameData = OpenVDBFramesData[FrameIndex];
-			AnimationAABBMin = FrameIndex == 0 ? OpenVDBFrameData.OpenVDBData.VolumeActiveAABBMin : FVector::Min(AnimationAABBMin, OpenVDBFrameData.OpenVDBData.VolumeActiveAABBMin);
-			AnimationAABBMax = FrameIndex == 0 ? OpenVDBFrameData.OpenVDBData.VolumeActiveAABBMax : FVector::Max(AnimationAABBMin, OpenVDBFrameData.OpenVDBData.VolumeActiveAABBMax);
-		}
-
 		// Allocate space for each frame
 		AnimatedSVTexture->FrameCount = FrameCount;
 		AnimatedSVTexture->AnimationFrames.SetNum(FrameCount);
 
 		UE_LOG(LogSparseVolumeTextureFactory, Display, TEXT("Serializing: %i frame"), FrameCount);
 
-		// Convert all openvdb files to sparse volume texture
-		for (int32 FrameIndex = 0; FrameIndex < FrameCount; ++FrameIndex)
+		// Load vdb files for all frames.
+		for(int32 FrameIndex = 0; FrameIndex < FrameCount; ++FrameIndex)
 		{
-			FOpenVDBFrameData& OpenVDBFrameData = OpenVDBFramesData[FrameIndex];
-			FSparseVolumeRawSource SparseVolumeRawSource;
-			if (!ConvertOpenVDBToSparseVolumeTexture(OpenVDBFrameData.DensityGrid, OpenVDBFrameData.Filename, SparseVolumeRawSource,
-				true, AnimationAABBMin, AnimationAABBMax))
+			const FString FrameFilename = GetOpenVDBFileNameForFrame(FrameIndex);
+			TArray<uint8> LoadedFile;
+
+			if (!FFileHelper::LoadFileToArray(LoadedFile, *FrameFilename))
 			{
+				UE_LOG(LogSparseVolumeTextureFactory, Error, TEXT("OpenVDB file could not be opened: %s"), *FrameFilename);
 				return nullptr;
 			}
 
-			// Serialise the raw source data fro this frame into the asset object.
+			uint32 DensityGridIndex = 0;
+			FOpenVDBData OpenVDBData;
+			if (!FindDensityGridIndex(LoadedFile, FrameFilename, &DensityGridIndex, &OpenVDBData))
 			{
-				UE_LOG(LogSparseVolumeTextureFactory, Display, TEXT("  - frame %i (active dimension %i x %i x %i)"), FrameIndex, 
-					OpenVDBFrameData.OpenVDBData.VolumeActiveDim.X, OpenVDBFrameData.OpenVDBData.VolumeActiveDim.Y, OpenVDBFrameData.OpenVDBData.VolumeActiveDim.Z);
+				UE_LOG(LogSparseVolumeTextureFactory, Error, TEXT("OpenVDB file contains no suitable density grid: %s"), *FrameFilename);
+				return nullptr;
+			}
+
+			UE_LOG(LogSparseVolumeTextureFactory, Display, TEXT("  - frame %i (active dimension %i x %i x %i)"), FrameIndex,
+				OpenVDBData.VolumeActiveDim.X, OpenVDBData.VolumeActiveDim.Y, OpenVDBData.VolumeActiveDim.Z);
 
-				UE_LOG(LogSparseVolumeTextureFactory, Display, TEXT("    -  SourceVolumeResolution %i x %i x %i"),
-					SparseVolumeRawSource.Header.SourceVolumeResolution.X, SparseVolumeRawSource.Header.SourceVolumeResolution.Y, SparseVolumeRawSource.Header.SourceVolumeResolution.Z);
-				UE_LOG(LogSparseVolumeTextureFactory, Display, TEXT("    -  PageTableVolumeResolution %i x %i x %i"),
-					SparseVolumeRawSource.Header.PageTableVolumeResolution.X, SparseVolumeRawSource.Header.PageTableVolumeResolution.Y, SparseVolumeRawSource.Header.PageTableVolumeResolution.Z);
-				UE_LOG(LogSparseVolumeTextureFactory, Display, TEXT("    -  TileDataVolumeResolution %i x %i x %i"),
-					SparseVolumeRawSource.Header.TileDataVolumeResolution.X, SparseVolumeRawSource.Header.TileDataVolumeResolution.Y, SparseVolumeRawSource.Header.TileDataVolumeResolution.Z);
+			FSparseVolumeRawSource SparseVolumeRawSource{};
+			SparseVolumeRawSource.DensityGridIndex = DensityGridIndex;
+			SparseVolumeRawSource.SourceAssetFile = MoveTemp(LoadedFile);
+			LoadedFile.Reset();
 
+			// Serialize the raw source data from this frame into the asset object.
+			{
 				UE::Serialization::FEditorBulkDataWriter RawDataArchiveWriter(AnimatedSVTexture->AnimationFrames[FrameIndex].RawData);
 				SparseVolumeRawSource.Serialize(RawDataArchiveWriter);
-
-				if (ImportTask.ShouldCancel())
-				{
-					return nullptr;
-				}
-				ImportTask.EnterProgressFrame(1.0f, LOCTEXT("ConvertingVDBStatic", "Converting static OpenVDB"));
 			}
-		}
 
-		// Close all the files.
-		for (int32 FrameIndex = 0; FrameIndex < FrameCount; ++FrameIndex)
-		{
-			FOpenVDBFrameData& OpenVDBFrameData = OpenVDBFramesData[FrameIndex];
-			OpenVDBFrameData.File->close();
-			delete OpenVDBFrameData.File;
+			if (ImportTask.ShouldCancel())
+			{
+				return nullptr;
+			}
+			ImportTask.EnterProgressFrame(1.0f, LOCTEXT("ConvertingVDBAnim", "Converting OpenVDB animation"));
 		}
 
 		ResultAssets.Add(AnimatedSVTexture);
@@ -521,63 +215,33 @@ UObject* USparseVolumeTextureFactory::FactoryCreateFile(UClass* InClass, UObject
 		FName NewName(InName.ToString() + TEXT("VDB"));
 		UStaticSparseVolumeTexture* StaticSVTexture = NewObject<UStaticSparseVolumeTexture>(InParent, UStaticSparseVolumeTexture::StaticClass(), NewName, Flags);
 
-		std::string FileNameStr(TCHAR_TO_ANSI(*Filename));
-		openvdb::io::File File(FileNameStr.c_str());
-
-		try
-		{
-			File.open();
-		}
-		catch (openvdb::Exception e)
+		TArray<uint8> LoadedFile;
+		if (!FFileHelper::LoadFileToArray(LoadedFile, *Filename))
 		{
-			UE_LOG(LogSparseVolumeTextureFactory, Warning, TEXT("OpenVDB file coult not be opened: %s"), *Filename);
+			UE_LOG(LogSparseVolumeTextureFactory, Error, TEXT("OpenVDB file could not be opened: %s"), *Filename);
 			return nullptr;
 		}
 
 		FScopedSlowTask ImportTask(1.0f, LOCTEXT("ImportingVDBStatic", "Importing static OpenVDB"));
 		ImportTask.MakeDialog(true);
 
-		openvdb::GridBase::Ptr BaseGrid = nullptr;
-		for (openvdb::io::File::NameIterator nameIter = File.beginName();
-			nameIter != File.endName(); ++nameIter)
-		{
-			if (nameIter.gridName() == "density")
-			{
-				openvdb::GridBase::Ptr QueryGrid = File.readGrid(nameIter.gridName());
-				if (QueryGrid->isType<openvdb::FloatGrid>())
-				{
-					BaseGrid = QueryGrid;
-					break;
-				}
-				break;
-			}
-		}
-
-		// If we have not found any density map, let's pick up the first float map
-		for (openvdb::io::File::NameIterator nameIter = File.beginName();
-			nameIter != File.endName(); ++nameIter)
+		uint32 DensityGridIndex = 0;
+		FOpenVDBData OpenVDBData;
+		if (!FindDensityGridIndex(LoadedFile, Filename, &DensityGridIndex, &OpenVDBData))
 		{
-			openvdb::GridBase::Ptr QueryGrid = File.readGrid(nameIter.gridName());
-			if (QueryGrid->isType<openvdb::FloatGrid>())
-			{
-				BaseGrid = QueryGrid;
-				break;
-			}
-		}
-
-		openvdb::FloatGrid::Ptr DensityGrid = openvdb::gridPtrCast<openvdb::FloatGrid>(BaseGrid);
-
-		// Only open float grid for now
-		if (DensityGrid == nullptr)
-		{
-			UE_LOG(LogSparseVolumeTextureFactory, Error, TEXT("Could not convert the first grid to float: %s"), *Filename);
+			UE_LOG(LogSparseVolumeTextureFactory, Error, TEXT("OpenVDB file contains no suitable density grid: %s"), *Filename);
 			return nullptr;
 		}
 
-		FSparseVolumeRawSource SparseVolumeRawSource;
-		if (!ConvertOpenVDBToSparseVolumeTexture(DensityGrid, Filename, SparseVolumeRawSource))
+		FSparseVolumeRawSource SparseVolumeRawSource{};
+		SparseVolumeRawSource.DensityGridIndex = DensityGridIndex;
+		SparseVolumeRawSource.SourceAssetFile = MoveTemp(LoadedFile);
+		LoadedFile.Reset();
+
+		// Serialize the raw source data into the asset object.
 		{
-			return nullptr;
+			UE::Serialization::FEditorBulkDataWriter RawDataArchiveWriter(StaticSVTexture->StaticFrame.RawData);
+			SparseVolumeRawSource.Serialize(RawDataArchiveWriter);
 		}
 
 		if (ImportTask.ShouldCancel())
@@ -586,18 +250,9 @@ UObject* USparseVolumeTextureFactory::FactoryCreateFile(UClass* InClass, UObject
 		}
 		ImportTask.EnterProgressFrame(1.0f, LOCTEXT("ConvertingVDBStatic", "Converting static OpenVDB"));
 
-		File.close();
-
-		// Serialise the raw source data into the asset object.
-		{
-			UE::Serialization::FEditorBulkDataWriter RawDataArchiveWriter(StaticSVTexture->StaticFrame.RawData);
-			SparseVolumeRawSource.Serialize(RawDataArchiveWriter);
-		}
 		ResultAssets.Add(StaticSVTexture);
 	}
 
-
-
 	// Now notify the system about the imported/updated/created assets
 	AdditionalImportedObjects.Reserve(ResultAssets.Num());
 	for (UObject* Object : ResultAssets)
diff --git a/Engine/Source/Editor/SparseVolumeTexture/Private/SparseVolumeTextureModule.cpp b/Engine/Source/Editor/SparseVolumeTexture/Private/SparseVolumeTextureModule.cpp
--- a/Engine/Source/Editor/SparseVolumeTexture/Private/SparseVolumeTextureModule.cpp
+++ b/Engine/Source/Editor/SparseVolumeTexture/Private/SparseVolumeTextureModule.cpp
@@ -3,6 +3,8 @@
 #include "SparseVolumeTextureModule.h"
 
 #include "SparseVolumeTextureOpenVDB.h"
+#include "SparseVolumeTexture/SparseVolumeTexture.h"
+#include "SparseVolumeTextureOpenVDBUtility.h"
 
 #define LOCTEXT_NAMESPACE "SparseVolumeTextureModule"
 
@@ -18,10 +20,14 @@ void FSparseVolumeTextureModule::StartupModule()
 	// Global registration of  the vdb types.
 	openvdb::initialize();
 #endif
+
+	// USparseVolumeTexture needs to access (only) this function for cooking.
+	OnConvertOpenVDBToSparseVolumeTexture().BindStatic(ConvertOpenVDBToSparseVolumeTexture);
 }
 
 void FSparseVolumeTextureModule::ShutdownModule()
 {
+	OnConvertOpenVDBToSparseVolumeTexture().Unbind();
 }
 
 #undef LOCTEXT_NAMESPACE
diff --git a/Engine/Source/Editor/SparseVolumeTexture/Private/SparseVolumeTextureOpenVDBUtility.cpp b/Engine/Source/Editor/SparseVolumeTexture/Private/SparseVolumeTextureOpenVDBUtility.cpp
new file mode 100644
--- /dev/null
+++ b/Engine/Source/Editor/SparseVolumeTexture/Private/SparseVolumeTextureOpenVDBUtility.cpp
@@ -0,0 +1,332 @@
+// Copyright Epic Games, Inc. All Rights Reserved.
+
+#if WITH_EDITOR
+
+#include "SparseVolumeTextureOpenVDBUtility.h"
+#include "SparseVolumeTextureOpenVDB.h"
+#include "SparseVolumeTexture/SparseVolumeTexture.h"
+
+DEFINE_LOG_CATEGORY_STATIC(LogSparseVolumeTextureOpenVDBUtility, Log, All);
+
+#if OPENVDB_AVAILABLE
+
+namespace
+{
+	// Utility class acting as adapter between TArray<uint8> and std::istream
+	class FArrayUint8StreamBuf : public std::streambuf
+	{
+	public:
+		explicit FArrayUint8StreamBuf(TArray<uint8>& Array)
+		{
+			char* Data = (char*)Array.GetData();
+			setg(Data, Data, Data + Array.Num());
+		}
+	};
+}
+
+static FOpenVDBData GetOpenVDBData(openvdb::GridBase::Ptr GridBase)
+{
+	FOpenVDBData OpenVDBData;
+	openvdb::CoordBBox VolumeActiveAABB = GridBase->evalActiveVoxelBoundingBox();
+	openvdb::Coord VolumeActiveDim = GridBase->evalActiveVoxelDim();
+	openvdb::Vec3d VolumeVoxelSize = GridBase->voxelSize();
+
+	OpenVDBData.VolumeActiveAABBMin = FVector(VolumeActiveAABB.min().x(), VolumeActiveAABB.min().y(), VolumeActiveAABB.min().z());
+	OpenVDBData.VolumeActiveAABBMax = FVector(VolumeActiveAABB.max().x(), VolumeActiveAABB.max().y(), VolumeActiveAABB.max().z());
+	OpenVDBData.VolumeActiveDim = FVector(VolumeActiveDim.x(), VolumeActiveDim.y(), VolumeActiveDim.z());
+	OpenVDBData.VolumeVoxelSize = FVector(VolumeVoxelSize.x(), VolumeVoxelSize.y(), VolumeVoxelSize.z());
+	OpenVDBData.bIsInWorldSpace = GridBase->isInWorldSpace();
+	OpenVDBData.bHasUniformVoxels = GridBase->hasUniformVoxels();
+
+	return OpenVDBData;
+}
+
+#endif
+
+bool IsOpenVDBDataValid(FOpenVDBData& OpenVDBData, const FString& Filename)
+{
+	if (OpenVDBData.VolumeActiveDim.X * OpenVDBData.VolumeActiveDim.Y * OpenVDBData.VolumeActiveDim.Z == 0)
+	{
+		// SVT_TODO we should gently handle that case
+		UE_LOG(LogSparseVolumeTextureOpenVDBUtility, Error, TEXT("OpenVDB asset is empty due to volume size being 0: %s"), *Filename);
+		return false;
+	}
+
+	if (!OpenVDBData.bHasUniformVoxels)
+	{
+		UE_LOG(LogSparseVolumeTextureOpenVDBUtility, Error, TEXT("OpenVDB importer cannot handle non uniform voxels: %s"), *Filename);
+		return false;
+	}
+	return true;
+}
+
+bool FindDensityGridIndex(TArray<uint8>& SourceFile, const FString& Filename, uint32* OutGridIndex, FOpenVDBData* OutOVDBData)
+{
+#if OPENVDB_AVAILABLE
+	FArrayUint8StreamBuf StreamBuf(SourceFile);
+	std::istream IStream(&StreamBuf);
+	openvdb::io::Stream Stream(IStream, false /*delayLoad*/);
+
+	openvdb::GridPtrVecPtr Grids = Stream.getGrids();
+
+	if (!Grids)
+	{
+		UE_LOG(LogSparseVolumeTextureOpenVDBUtility, Warning, TEXT("OpenVDB file contains no grids: %s"), *Filename);
+		return false;
+	}
+
+	uint32 GridIndex = 0;
+	openvdb::GridBase::Ptr BaseGrid = nullptr;
+	for (openvdb::GridBase::Ptr& Grid : *Grids)
+	{
+		if (Grid->getName() == "density")
+		{
+			if (Grid->isType<openvdb::FloatGrid>())
+			{
+				BaseGrid = Grid;
+				break;
+			}
+			break;
+		}
+		++GridIndex;
+	}
+
+	// If we have not found any density map, let's pick up the first float map
+	if (!BaseGrid)
+	{
+		GridIndex = 0;
+		for (openvdb::GridBase::Ptr& Grid : *Grids)
+		{
+			if (Grid->isType<openvdb::FloatGrid>())
+			{
+				BaseGrid = Grid;
+				break;
+			}
+			++GridIndex;
+		}
+	}
+
+	if (BaseGrid)
+	{
+		openvdb::FloatGrid::Ptr DensityGrid = openvdb::gridPtrCast<openvdb::FloatGrid>(BaseGrid);
+
+		// Only open float grid for now
+		if (DensityGrid == nullptr)
+		{
+			UE_LOG(LogSparseVolumeTextureOpenVDBUtility, Error, TEXT("Could not convert the grid to float: %s"), *Filename);
+			return false;
+		}
+
+		FOpenVDBData OVDBData = GetOpenVDBData(DensityGrid);
+		if (!IsOpenVDBDataValid(OVDBData, Filename))
+		{
+			return false;
+		}
+
+		*OutGridIndex = GridIndex;
+		*OutOVDBData = OVDBData;
+		return true;
+	}
+
+#endif // OPENVDB_AVAILABLE
+
+	return false;
+}
+
+bool ConvertOpenVDBToSparseVolumeTexture(
+	TArray<uint8>& SourceFile,
+	uint32 GridIndex,
+	FSparseVolumeAssetHeader* OutHeader,
+	TArray<uint32>* OutDensityPage,
+	TArray<uint8>* OutDensityData,
+	bool bOverrideActiveMinMax,
+	FVector ActiveMin,
+	FVector ActiveMax)
+{
+#if OPENVDB_AVAILABLE
+	FArrayUint8StreamBuf StreamBuf(SourceFile);
+	std::istream IStream(&StreamBuf);
+	openvdb::io::Stream Stream(IStream, false /*delayLoad*/);
+
+	openvdb::GridPtrVecPtr Grids = Stream.getGrids();
+	if (!Grids || Grids->size() <= GridIndex)
+	{
+		return false;
+	}
+
+	openvdb::FloatGrid::Ptr DensityGrid = openvdb::gridPtrCast<openvdb::FloatGrid>((*Grids)[GridIndex]);
+	if (!DensityGrid)
+	{
+		return false;
+	}
+
+	FOpenVDBData OVDBData = GetOpenVDBData(DensityGrid);
+	if (!IsOpenVDBDataValid(OVDBData, TEXT("")))
+	{
+		return false;
+	}
+
+	if (bOverrideActiveMinMax)
+	{
+		OVDBData.VolumeActiveAABBMin = ActiveMin;
+		OVDBData.VolumeActiveAABBMax = ActiveMax;
+		OVDBData.VolumeActiveDim = ActiveMax - ActiveMin;
+	}
+
+	FSparseVolumeAssetHeader& Header = *OutHeader;
+
+	Header.SourceVolumeResolution = FIntVector(OVDBData.VolumeActiveDim);
+	// SVT_TODO handle error too high volume resolution assuming dx11 limitations?
+
+	Header.PageTableVolumeResolution = FIntVector3(
+		FMath::DivideAndRoundUp(Header.SourceVolumeResolution.X, SPARSE_VOLUME_TILE_RES),
+		FMath::DivideAndRoundUp(Header.SourceVolumeResolution.Y, SPARSE_VOLUME_TILE_RES),
+		FMath::DivideAndRoundUp(Header.SourceVolumeResolution.Z, SPARSE_VOLUME_TILE_RES));
+	Header.TileDataVolumeResolution = FIntVector::ZeroValue;	// unknown for now
+
+	const uint32 VolumeTileDataBytes = SPARSE_VOLUME_TILE_RES * SPARSE_VOLUME_TILE_RES * SPARSE_VOLUME_TILE_RES;
+	uint32 NumAllocatedPages = 0;
+
+	// Allocate some memory for temp data (worst case)
+	TArray<uint8> LinearAllocatedTiles;
+	TArray<FIntVector3> LinearAllocatedPages;
+	LinearAllocatedTiles.SetNum(Header.PageTableVolumeResolution.X * Header.PageTableVolumeResolution.Y * Header.PageTableVolumeResolution.Z * VolumeTileDataBytes);
+	LinearAllocatedPages.SetNum(Header.PageTableVolumeResolution.X * Header.PageTableVolumeResolution.Y * Header.PageTableVolumeResolution.Z);
+
+	auto accessor = DensityGrid->getConstAccessor();
+	const float GridBackgroundValue = DensityGrid->background();
+
+	// Go over each potential page table form the source data and push allocate it if it has any data.
+	// Otherwise point to the default empty page.
+	bool bAnyEmptyPageExists = false;
+	for (int32_t PageZ = 0; PageZ < Header.PageTableVolumeResolution.Z; ++PageZ)
+	{
+		for (int32_t PageY = 0; PageY < Header.PageTableVolumeResolution.Y; ++PageY)
+		{
+			for (int32_t PageX = 0; PageX < Header.PageTableVolumeResolution.X; ++PageX)
+			{
+				bool bHasAnyData = false;
+
+				uint8* NewTileData = LinearAllocatedTiles.GetData() + NumAllocatedPages * VolumeTileDataBytes;
+
+				for (int32_t z = 0; z < SPARSE_VOLUME_TILE_RES; ++z)
+				{
+					for (int32_t y = 0; y < SPARSE_VOLUME_TILE_RES; ++y)
+					{
+						for (int32_t x = 0; x < SPARSE_VOLUME_TILE_RES; ++x)
+						{
+							FVector VoxelCoord(PageX * SPARSE_VOLUME_TILE_RES + x, PageY * SPARSE_VOLUME_TILE_RES + y, PageZ * SPARSE_VOLUME_TILE_RES + z);	// This assumes sampling outside the boundary retuirn a default value
+							FVector VoxelCoord2 = OVDBData.VolumeActiveAABBMin + VoxelCoord;
+							float VoxelValue = accessor.getValue(openvdb::Coord(VoxelCoord2.X, VoxelCoord2.Y, VoxelCoord2.Z));
+
+							// clamp for now until we have fp16
+							VoxelValue = FMath::Clamp(VoxelValue, 0.0f, 1.0f);
+
+							bHasAnyData |= VoxelValue > GridBackgroundValue;
+
+							NewTileData[z * SPARSE_VOLUME_TILE_RES * SPARSE_VOLUME_TILE_RES + y * SPARSE_VOLUME_TILE_RES + x] = uint8(VoxelValue * 255.0f);
+						}
+					}
+				}
+
+				if (bHasAnyData)
+				{
+					LinearAllocatedPages[NumAllocatedPages] = FIntVector3(PageX, PageY, PageZ);
+					NumAllocatedPages++;
+				}
+				bAnyEmptyPageExists |= !bHasAnyData;
+			}
+		}
+	}
+
+	// Compute Page and Tile VolumeResolution from allocated pages
+	const uint32 EffectivelyAllocatedPageEntries = NumAllocatedPages + (bAnyEmptyPageExists ? 1 : 0);
+	uint32 TileVolumeResolutionCube = 1;
+	while (TileVolumeResolutionCube * TileVolumeResolutionCube * TileVolumeResolutionCube < EffectivelyAllocatedPageEntries)
+	{
+		TileVolumeResolutionCube++;				// We use a simple loop to compute the minimum resolution of a cube to store all the tile data
+	}
+	Header.TileDataVolumeResolution = FIntVector3(TileVolumeResolutionCube, TileVolumeResolutionCube, TileVolumeResolutionCube);
+	while (Header.TileDataVolumeResolution.X * Header.TileDataVolumeResolution.Y * (Header.TileDataVolumeResolution.Z - 1) > int32(EffectivelyAllocatedPageEntries))
+	{
+		Header.TileDataVolumeResolution.Z--;	// We then trim an edge to get back space.
+	}
+	const FIntVector3 TileCoordResolution = Header.TileDataVolumeResolution;
+	Header.TileDataVolumeResolution = Header.TileDataVolumeResolution * SPARSE_VOLUME_TILE_RES;
+
+	// Initialise the SparseVolumeTexture page and tile.
+	OutDensityPage->SetNumZeroed(Header.PageTableVolumeResolution.X * Header.PageTableVolumeResolution.Y * Header.PageTableVolumeResolution.Z);
+	OutDensityData->SetNumZeroed(Header.TileDataVolumeResolution.X * Header.TileDataVolumeResolution.Y * Header.TileDataVolumeResolution.Z);
+	uint32* DensityPagePtr = OutDensityPage->GetData();
+	uint8* DensityDataPtr = OutDensityData->GetData();
+
+	FIntVector DestinationTileCoord = FIntVector::ZeroValue;
+	auto GoToNextTileCoord = [&]()
+	{
+		DestinationTileCoord.X++;
+		if (DestinationTileCoord.X >= TileCoordResolution.X)
+		{
+			DestinationTileCoord.X = 0;
+			DestinationTileCoord.Y++;
+		}
+		if (DestinationTileCoord.Y >= TileCoordResolution.Y)
+		{
+			DestinationTileCoord.Y = 0;
+			DestinationTileCoord.Z++;
+		}
+	};
+
+	// Add an empty tile is needed, reserve slot at coord 0
+	if (bAnyEmptyPageExists)
+	{
+		// DensityPage is all cleared to zero, simply skip a tile
+		GoToNextTileCoord();
+	}
+
+	// Generate page table and tile volume data by splatting the data
+	for (uint32 i = 0; i < NumAllocatedPages; ++i)
+	{
+		FIntVector3 PageCoordToSplat = LinearAllocatedPages[i];
+		uint8* TileDataToSplat = LinearAllocatedTiles.GetData() + i * VolumeTileDataBytes;
+
+		// A page encodes the physical tile coord as unsigned int of 11 11 10 bits
+		// This means a page coord cannot be larger than 2047 for x and y and 1023 for z
+		// which mean we cannot have more than 2048*2048*1024 = 4 Giga tiles of 16^3 tiles.
+		uint32 DestinationTileCoord32bit = (DestinationTileCoord.X & 0x7FF) | ((DestinationTileCoord.Y & 0x7FF) << 11) | ((DestinationTileCoord.Z & 0x3FF) << 22);
+
+		// Setup the page table entry
+		DensityPagePtr
+			[
+				PageCoordToSplat.Z * Header.PageTableVolumeResolution.X * Header.PageTableVolumeResolution.Y +
+				PageCoordToSplat.Y * Header.PageTableVolumeResolution.X +
+			PageCoordToSplat.X
+			] = DestinationTileCoord32bit;
+
+		// Now copy the tile data
+		for (int32_t z = 0; z < SPARSE_VOLUME_TILE_RES; ++z)
+		{
+			for (int32_t y = 0; y < SPARSE_VOLUME_TILE_RES; ++y)
+			{
+				for (int32_t x = 0; x < SPARSE_VOLUME_TILE_RES; ++x)
+				{
+					DensityDataPtr
+						[
+							(DestinationTileCoord.Z * SPARSE_VOLUME_TILE_RES + z) * Header.TileDataVolumeResolution.X * Header.TileDataVolumeResolution.Y +
+							(DestinationTileCoord.Y * SPARSE_VOLUME_TILE_RES + y) * Header.TileDataVolumeResolution.X +
+						(DestinationTileCoord.X * SPARSE_VOLUME_TILE_RES + x)
+						] = TileDataToSplat[z * SPARSE_VOLUME_TILE_RES * SPARSE_VOLUME_TILE_RES + y * SPARSE_VOLUME_TILE_RES + x];
+				}
+			}
+		}
+
+		// Set the next tile to be written to
+		GoToNextTileCoord();
+	}
+
+	return true;
+#else
+	return false;
+#endif // OPENVDB_AVAILABLE
+}
+
+#endif // WITH_EDITOR
\ No newline at end of file
diff --git a/Engine/Source/Editor/SparseVolumeTexture/Private/SparseVolumeTextureOpenVDBUtility.h b/Engine/Source/Editor/SparseVolumeTexture/Private/SparseVolumeTextureOpenVDBUtility.h
new file mode 100644
--- /dev/null
+++ b/Engine/Source/Editor/SparseVolumeTexture/Private/SparseVolumeTextureOpenVDBUtility.h
@@ -0,0 +1,33 @@
+// Copyright Epic Games, Inc. All Rights Reserved.
+
+#pragma once
+
+#if WITH_EDITOR
+
+#include "CoreMinimal.h"
+
+struct FOpenVDBData
+{
+	FVector VolumeActiveAABBMin;
+	FVector VolumeActiveAABBMax;
+	FVector VolumeActiveDim;
+	FVector VolumeVoxelSize;
+	bool bIsInWorldSpace;
+	bool bHasUniformVoxels;
+};
+
+bool IsOpenVDBDataValid(FOpenVDBData& OpenVDBData, const FString& Filename);
+
+bool FindDensityGridIndex(TArray<uint8>& SourceFile, const FString& Filename, uint32* OutGridIndex, FOpenVDBData* OutOVDBData);
+
+bool ConvertOpenVDBToSparseVolumeTexture(
+	TArray<uint8>& SourceFile,
+	uint32 GridIndex,
+	struct FSparseVolumeAssetHeader* OutHeader,
+	TArray<uint32>* OutDensityPage,
+	TArray<uint8>* OutDensityData,
+	bool bOverrideActiveMinMax,
+	FVector ActiveMin,
+	FVector ActiveMax);
+
+#endif // WITH_EDITOR
\ No newline at end of file
diff --git a/Engine/Source/Runtime/Engine/Classes/SparseVolumeTexture/SparseVolumeTexture.h b/Engine/Source/Runtime/Engine/Classes/SparseVolumeTexture/SparseVolumeTexture.h
--- a/Engine/Source/Runtime/Engine/Classes/SparseVolumeTexture/SparseVolumeTexture.h
+++ b/Engine/Source/Runtime/Engine/Classes/SparseVolumeTexture/SparseVolumeTexture.h
@@ -17,6 +17,18 @@ namespace UE { namespace Shader	{ enum class EValueType : uint8; } }
 
 #define SPARSE_VOLUME_TILE_RES	16
 
+DECLARE_DELEGATE_RetVal_EightParams(bool, FConvertOpenVDBToSparseVolumeTextureDelegate, 
+	TArray<uint8>& SourceFile,
+	uint32 GridIndex,
+	struct FSparseVolumeAssetHeader* OutHeader,
+	TArray<uint32>* OutDensityPage,
+	TArray<uint8>* OutDensityData,
+	bool bOverrideActiveMinMax,
+	FVector ActiveMin,
+	FVector ActiveMax);
+
+ENGINE_API FConvertOpenVDBToSparseVolumeTextureDelegate& OnConvertOpenVDBToSparseVolumeTexture();
+
 struct ENGINE_API FSparseVolumeAssetHeader
 {
 	FIntVector3 PageTableVolumeResolution;
@@ -43,12 +55,8 @@ struct ENGINE_API FSparseVolumeAssetHeader
 // The structure represent the source asset in high quality. It is used to cook the runtime data
 struct ENGINE_API FSparseVolumeRawSource
 {
-	FSparseVolumeAssetHeader Header;
-
-	// SVT_TODO Store the source asset as the saved openvdb instead of page table/physical tile as it is now.
-
-	TArray<uint32>	DensityPage;
-	TArray<uint8>	DensityData;
+	uint32						DensityGridIndex;
+	TArray<uint8>				SourceAssetFile;
 
 	// The current data format version for the raw source data.
 	static const uint32 kVersion = 0;
@@ -59,7 +67,7 @@ struct ENGINE_API FSparseVolumeRawSource
 	void Serialize(FArchive& Ar);
 
 	FSparseVolumeRawSource()
-		: Header()
+		: DensityGridIndex()
 		, Version(kVersion)
 	{
 	}
diff --git a/Engine/Source/Runtime/Engine/Private/SparseVolumeTexture/SparseVolumeTexture.cpp b/Engine/Source/Runtime/Engine/Private/SparseVolumeTexture/SparseVolumeTexture.cpp
--- a/Engine/Source/Runtime/Engine/Private/SparseVolumeTexture/SparseVolumeTexture.cpp
+++ b/Engine/Source/Runtime/Engine/Private/SparseVolumeTexture/SparseVolumeTexture.cpp
@@ -26,6 +26,13 @@
 
 DEFINE_LOG_CATEGORY_STATIC(LogSparseVolumeTexture, Log, All);
 
+FConvertOpenVDBToSparseVolumeTextureDelegate ConvertOpenVDBToSparseVolumeTextureDelegate;
+
+FConvertOpenVDBToSparseVolumeTextureDelegate& OnConvertOpenVDBToSparseVolumeTexture()
+{
+	return ConvertOpenVDBToSparseVolumeTextureDelegate;
+}
+
 void FSparseVolumeAssetHeader::Serialize(FArchive& Ar)
 {
 
@@ -46,14 +53,12 @@ void FSparseVolumeAssetHeader::Serialize(FArchive& Ar)
 
 void FSparseVolumeRawSource::Serialize(FArchive& Ar)
 {
-	Header.Serialize(Ar);
-
 	Ar << Version;
 
 	if (Version == 0)
 	{
-		Ar << DensityData;
-		Ar << DensityPage;
+		Ar << DensityGridIndex;
+		Ar << SourceAssetFile;
 	}
 	else
 	{
@@ -226,9 +231,14 @@ void UStaticSparseVolumeTexture::ConvertRawSourceDataToSparseVolumeTextureRuntim
 		SparseVolumeRawSource.Serialize(RawDataArchiveReader);
 
 		// Then, cook the runtime data and serialize it out.
-		StaticFrame.SparseVolumeTextureRuntime.Header = SparseVolumeRawSource.Header;
-		StaticFrame.SparseVolumeTextureRuntime.DensityData = SparseVolumeRawSource.DensityData;
-		StaticFrame.SparseVolumeTextureRuntime.DensityPage = SparseVolumeRawSource.DensityPage;
+		const bool bSuccess = ConvertOpenVDBToSparseVolumeTextureDelegate.IsBound() && ConvertOpenVDBToSparseVolumeTextureDelegate.Execute(
+			SparseVolumeRawSource.SourceAssetFile,
+			SparseVolumeRawSource.DensityGridIndex,
+			&StaticFrame.SparseVolumeTextureRuntime.Header,
+			&StaticFrame.SparseVolumeTextureRuntime.DensityPage,
+			&StaticFrame.SparseVolumeTextureRuntime.DensityData,
+			false, FVector::Zero(), FVector::Zero());
+		ensure(bSuccess);
 
 		// Now unload the raw data
 		StaticFrame.RawData.UnloadData();
@@ -524,9 +534,14 @@ void UAnimatedSparseVolumeTexture::ConvertRawSourceDataToSparseVolumeTextureRunt
 		SparseVolumeRawSource.Serialize(RawDataArchiveReader);
 
 		// Then, cook the runtime data and serialize it out.
-		Frame.SparseVolumeTextureRuntime.Header = SparseVolumeRawSource.Header;
-		Frame.SparseVolumeTextureRuntime.DensityData = SparseVolumeRawSource.DensityData;
-		Frame.SparseVolumeTextureRuntime.DensityPage = SparseVolumeRawSource.DensityPage;
+		const bool bSuccess = ConvertOpenVDBToSparseVolumeTextureDelegate.IsBound() && ConvertOpenVDBToSparseVolumeTextureDelegate.Execute(
+			SparseVolumeRawSource.SourceAssetFile,
+			SparseVolumeRawSource.DensityGridIndex,
+			&Frame.SparseVolumeTextureRuntime.Header,
+			&Frame.SparseVolumeTextureRuntime.DensityPage,
+			&Frame.SparseVolumeTextureRuntime.DensityData,
+			false, FVector::Zero(), FVector::Zero());
+		ensure(bSuccess);
 
 		// Now unload the raw data
 		Frame.RawData.UnloadData();
