Implemented first version of OpenVDB import dialog.

#rnx

diff --git a/Engine/Source/Editor/SparseVolumeTexture/Private/OpenVDBImportOptions.h b/Engine/Source/Editor/SparseVolumeTexture/Private/OpenVDBImportOptions.h
new file mode 100644
--- /dev/null
+++ b/Engine/Source/Editor/SparseVolumeTexture/Private/OpenVDBImportOptions.h
@@ -0,0 +1,17 @@
+// Copyright Epic Games, Inc. All Rights Reserved.
+
+#pragma once
+
+#include "CoreMinimal.h"
+
+struct FOpenVDBImportChannel
+{
+	FString Name;
+	uint32 Index;
+	bool bImport;
+};
+
+struct FOpenVDBImportOptions
+{
+	FOpenVDBImportChannel Density;
+};
diff --git a/Engine/Source/Editor/SparseVolumeTexture/Private/OpenVDBImportWindow.cpp b/Engine/Source/Editor/SparseVolumeTexture/Private/OpenVDBImportWindow.cpp
new file mode 100644
--- /dev/null
+++ b/Engine/Source/Editor/SparseVolumeTexture/Private/OpenVDBImportWindow.cpp
@@ -0,0 +1,327 @@
+#include "OpenVDBImportWindow.h"
+
+#include "Widgets/Input/SButton.h"
+#include "SPrimaryButton.h"
+#include "Widgets/Layout/SBorder.h"
+#include "Widgets/Layout/SBox.h"
+#include "Widgets/Layout/SUniformGridPanel.h"
+#include "Widgets/SBoxPanel.h"
+#include "Widgets/Text/STextBlock.h"
+
+#include "IDocumentation.h"
+#include "Editor.h"
+#include "SparseVolumeTextureOpenVDBUtility.h"
+#include "OpenVDBImportOptions.h"
+
+#define LOCTEXT_NAMESPACE "SOpenVDBImportWindow"
+
+static FText GetGridComboBoxItemText(TSharedPtr<FOpenVDBGridInfo> InItem)
+{
+	return InItem ? FText::FromString(InItem->DisplayString) : LOCTEXT("NoneGrid", "<None>");
+};
+
+void SOpenVDBImportWindow::Construct(const FArguments& InArgs)
+{
+	ImportOptions = InArgs._ImportOptions;
+	OpenVDBGridInfo = InArgs._OpenVDBGridInfo;
+	WidgetWindow = InArgs._WidgetWindow;
+
+	TSharedPtr<SBox> ImportTypeDisplay;
+	TSharedPtr<SHorizontalBox> OpenVDBHeaderButtons;
+	TSharedPtr<SBox> InspectorBox;
+	this->ChildSlot
+	[
+		SNew(SBox)
+		.MaxDesiredHeight(InArgs._MaxWindowHeight)
+		.MaxDesiredWidth(InArgs._MaxWindowWidth)
+		[
+			SNew(SVerticalBox)
+			+ SVerticalBox::Slot()
+			.AutoHeight()
+			.Padding(2)
+			[
+				SAssignNew(ImportTypeDisplay, SBox)
+			]
+			+ SVerticalBox::Slot()
+			.AutoHeight()
+			.Padding(2)
+			[
+				SNew(SBorder)
+				.Padding(FMargin(3))
+				.BorderImage(FAppStyle::GetBrush("ToolPanel.GroupBorder"))
+				[
+					SNew(SHorizontalBox)
+					+ SHorizontalBox::Slot()
+					.AutoWidth()
+					[
+						SNew(STextBlock)
+						.Text(LOCTEXT("Import_CurrentFileTitle", "Current Asset: "))
+					]
+					+ SHorizontalBox::Slot()
+					.Padding(5, 0, 0, 0)
+					.AutoWidth()
+					.VAlign(VAlign_Center)
+					[
+						SNew(STextBlock)
+						.Text(InArgs._FullPath)
+						.ToolTipText(InArgs._FullPath)
+					]
+				]
+			]
+			+ SVerticalBox::Slot()
+			.AutoHeight()
+			.Padding(2)
+			[
+				SAssignNew(InspectorBox, SBox)
+				.MaxDesiredHeight(650.0f)
+				.WidthOverride(400.0f)
+			]
+			+ SVerticalBox::Slot()
+			.AutoHeight()
+			.Padding(2)
+			[
+				SNew(SUniformGridPanel)
+				.SlotPadding(2)
+				+ SUniformGridPanel::Slot(1, 0)
+				[
+					SAssignNew(ImportButton, SPrimaryButton)
+					.Text(LOCTEXT("OpenVDBImportWindow_Import", "Import"))
+					.IsEnabled(this, &SOpenVDBImportWindow::CanImport)
+					.OnClicked(this, &SOpenVDBImportWindow::OnImport)
+				]
+				+ SUniformGridPanel::Slot(2, 0)
+				[
+					SNew(SButton)
+					.HAlign(HAlign_Center)
+					.Text(LOCTEXT("OpenVDBImportWindow_Cancel", "Cancel"))
+					.ToolTipText(LOCTEXT("OpenVDBImportWindow_Cancel_ToolTip", "Cancels importing this OpenVDB file"))
+					.OnClicked(this, &SOpenVDBImportWindow::OnCancel)
+				]
+			]
+		]
+	];
+
+	InspectorBox->SetContent(
+		SNew(SVerticalBox)
+		+ SVerticalBox::Slot()
+		.AutoHeight()
+		.Padding(2)
+		[
+			CreateGridSelector(DensityGridComboBox, DensityGridCheckBox, ImportOptions->Density, LOCTEXT("DensityGridIndex", "Density"))
+		]
+	);
+
+	SetDefaultGridAssignment();
+
+	ImportTypeDisplay->SetContent(
+		SNew(SBorder)
+		.Padding(FMargin(3))
+		.BorderImage(FAppStyle::GetBrush("ToolPanel.GroupBorder"))
+		[
+			SNew(SHorizontalBox)
+			+ SHorizontalBox::Slot()
+			.AutoWidth()
+			.VAlign(VAlign_Center)
+			[
+				SNew(STextBlock)
+				.Text(this, &SOpenVDBImportWindow::GetImportTypeDisplayText)
+			]
+			+ SHorizontalBox::Slot()
+			.VAlign(VAlign_Center)
+			.AutoWidth()
+			[
+				IDocumentation::Get()->CreateAnchor(FString("Engine/Content/OpenVDB/ImportWindow"))
+			]
+			+ SHorizontalBox::Slot()
+			.HAlign(HAlign_Right)
+			[
+				SAssignNew(OpenVDBHeaderButtons, SHorizontalBox)
+				+ SHorizontalBox::Slot()
+				.AutoWidth()
+				.Padding(FMargin(2.0f, 0.0f))
+				[
+					SNew(SButton)
+					.Text(LOCTEXT("OpenVDBImportWindow_ResetOptions", "Reset to Default"))
+					.OnClicked(this, &SOpenVDBImportWindow::OnResetToDefaultClick)
+				]
+			]
+		]
+	);
+
+	RegisterActiveTimer(0.f, FWidgetActiveTimerDelegate::CreateSP(this, &SOpenVDBImportWindow::SetFocusPostConstruct));
+}
+
+FReply SOpenVDBImportWindow::OnImport()
+{
+	bShouldImport = true;
+	if (WidgetWindow.IsValid())
+	{
+		WidgetWindow.Pin()->RequestDestroyWindow();
+	}
+	return FReply::Handled();
+}
+
+FReply SOpenVDBImportWindow::OnCancel()
+{
+	bShouldImport = false;
+	if (WidgetWindow.IsValid())
+	{
+		WidgetWindow.Pin()->RequestDestroyWindow();
+	}
+	return FReply::Handled();
+}
+
+bool SOpenVDBImportWindow::ShouldImport() const
+{
+	return bShouldImport;
+}
+
+EActiveTimerReturnType SOpenVDBImportWindow::SetFocusPostConstruct(double InCurrentTime, float InDeltaTime)
+{
+	if (ImportButton.IsValid())
+	{
+		FSlateApplication::Get().SetKeyboardFocus(ImportButton, EFocusCause::SetDirectly);
+	}
+
+	return EActiveTimerReturnType::Stop;
+}
+
+bool SOpenVDBImportWindow::CanImport() const
+{
+	return DensityGridCheckBox->IsChecked();
+}
+
+FReply SOpenVDBImportWindow::OnResetToDefaultClick()
+{
+	SetDefaultGridAssignment();
+	return FReply::Handled();
+}
+
+FText SOpenVDBImportWindow::GetImportTypeDisplayText() const
+{
+	return LOCTEXT("OpenVDBImportWindow_ImportType", "Import Static OpenVDB");
+}
+
+void SOpenVDBImportWindow::SetDefaultGridAssignment()
+{
+	check(OpenVDBGridInfo);
+
+	auto SetDefaultGrid = [](const TArray<TSharedPtr<FOpenVDBGridInfo>>& Grids, const FString& Name, EOpenVDBGridFormat Format, 
+		TSharedPtr<SComboBox<TSharedPtr<FOpenVDBGridInfo>>> ComboBox, TSharedPtr<SCheckBox> CheckBox, FOpenVDBImportChannel& Channel)
+	{
+		TSharedPtr<FOpenVDBGridInfo> SelectedGrid = nullptr;
+		// Search by name first...
+		for (auto& Grid : Grids)
+		{
+			if (Grid->Name == Name)
+			{
+				SelectedGrid = Grid;
+				break;
+			}
+		}
+
+		// ...and if that fails, just take the first grid with matching format.
+		if (!SelectedGrid)
+		{
+			for (auto& Grid : Grids)
+			{
+				if (Grid->Format == Format)
+				{
+					SelectedGrid = Grid;
+					break;
+				}
+			}
+		}
+
+		if (SelectedGrid)
+		{
+			const bool bImport = SelectedGrid->Name == Name; // disable importing this grid if the name doesn't match
+			ComboBox->SetSelectedItem(SelectedGrid);
+			CheckBox->SetIsChecked(bImport);
+			Channel.Name = SelectedGrid->Name;
+			Channel.Index = SelectedGrid->Index;
+			Channel.bImport = bImport;
+		}
+		else
+		{
+			ComboBox->ClearSelection();
+			CheckBox->SetIsChecked(false);
+			Channel.Name.Reset();
+			Channel.Index = INDEX_NONE;
+			Channel.bImport = false;
+		}
+	};
+
+	SetDefaultGrid(*OpenVDBGridInfo, TEXT("density"), EOpenVDBGridFormat::Float, DensityGridComboBox, DensityGridCheckBox, ImportOptions->Density);
+}
+
+TSharedRef<SWidget> SOpenVDBImportWindow::CreateGridSelector(TSharedPtr<SComboBox<TSharedPtr<FOpenVDBGridInfo>>>& ComboBox, TSharedPtr<SCheckBox>& CheckBox, FOpenVDBImportChannel& Channel, const FText& Label)
+{
+	return
+		SNew(SHorizontalBox)
+		+ SHorizontalBox::Slot()
+		.VAlign(VAlign_Center)
+		.AutoWidth()
+		.Padding(2.0f)
+		[
+			SAssignNew(CheckBox, SCheckBox)
+			.IsChecked(true)
+			.OnCheckStateChanged_Lambda([&Channel](ECheckBoxState CheckBoxState)
+			{
+				Channel.bImport = CheckBoxState == ECheckBoxState::Checked;
+			})
+		]
+
+		+ SHorizontalBox::Slot()
+		.VAlign(VAlign_Center)
+		.FillWidth(1.0f)
+		.Padding(2.0f)
+		[
+			SNew(STextBlock)
+			.Text(Label)
+		]
+
+		+ SHorizontalBox::Slot()
+		.FillWidth(1.0f)
+		.VAlign(VAlign_Center)
+		.Padding(2.0f)
+		[
+			SNew(SBox)
+			.WidthOverride(300.0f)
+			[
+				SAssignNew(ComboBox, SComboBox<TSharedPtr<FOpenVDBGridInfo>>)
+				.OptionsSource(OpenVDBGridInfo)
+				.IsEnabled_Lambda([&CheckBox]()
+				{
+					return CheckBox->IsChecked();
+				})
+				.OnGenerateWidget_Lambda([](TSharedPtr<FOpenVDBGridInfo> InItem)
+				{
+					return SNew(STextBlock)
+					.Text(GetGridComboBoxItemText(InItem));
+				})
+				.OnSelectionChanged_Lambda([&Channel](TSharedPtr<FOpenVDBGridInfo> InItem, ESelectInfo::Type)
+				{
+					if (InItem)
+					{
+						Channel.Name = InItem->Name;
+						Channel.Index = InItem->Index;
+					}
+					else
+					{
+						Channel.Name.Reset();
+						Channel.Index = INDEX_NONE;
+					}
+				})
+				[
+					SNew(STextBlock)
+					.Text_Lambda([&ComboBox]()
+					{
+						return GetGridComboBoxItemText(ComboBox->GetSelectedItem());
+					})
+				]
+			]
+		];
+}
+
+#undef LOCTEXT_NAMESPACE
diff --git a/Engine/Source/Editor/SparseVolumeTexture/Private/OpenVDBImportWindow.h b/Engine/Source/Editor/SparseVolumeTexture/Private/OpenVDBImportWindow.h
new file mode 100644
--- /dev/null
+++ b/Engine/Source/Editor/SparseVolumeTexture/Private/OpenVDBImportWindow.h
@@ -0,0 +1,56 @@
+// Copyright Epic Games, Inc. All Rights Reserved.
+
+#pragma once
+
+#include "CoreMinimal.h"
+#include "Widgets/SCompoundWidget.h"
+#include "Widgets/Input/SComboBox.h"
+#include "Widgets/Input/SCheckBox.h"
+
+struct FOpenVDBGridInfo;
+struct FOpenVDBImportOptions;
+struct FOpenVDBImportChannel;
+
+class SOpenVDBImportWindow : public SCompoundWidget
+{
+public:
+	SLATE_BEGIN_ARGS(SOpenVDBImportWindow)
+		: _ImportOptions()
+		,_OpenVDBGridInfo()
+		, _WidgetWindow()
+		, _FullPath()
+		, _MaxWindowHeight(0.0f)
+		, _MaxWindowWidth(0.0f)
+	{}
+
+		SLATE_ARGUMENT(FOpenVDBImportOptions*, ImportOptions)
+		SLATE_ARGUMENT(TArray<TSharedPtr<FOpenVDBGridInfo>>*, OpenVDBGridInfo)
+		SLATE_ARGUMENT(TSharedPtr<SWindow>, WidgetWindow)
+		SLATE_ARGUMENT(FText, FullPath)
+		SLATE_ARGUMENT(float, MaxWindowHeight)
+		SLATE_ARGUMENT(float, MaxWindowWidth)
+	SLATE_END_ARGS()
+
+public:
+	virtual bool SupportsKeyboardFocus() const override { return true; }
+	void Construct(const FArguments& InArgs);
+	FReply OnImport();
+	FReply OnCancel();
+	bool ShouldImport() const;
+
+private:
+	FOpenVDBImportOptions*								ImportOptions;
+	TArray<TSharedPtr<FOpenVDBGridInfo>>*				OpenVDBGridInfo;
+	TSharedPtr<SComboBox<TSharedPtr<FOpenVDBGridInfo>>>	DensityGridComboBox;
+	TSharedPtr<SCheckBox>								DensityGridCheckBox;
+	TSharedPtr<SButton>									ImportButton;
+	TWeakPtr<SWindow>									WidgetWindow;
+	bool												bShouldImport;
+
+	EActiveTimerReturnType SetFocusPostConstruct(double InCurrentTime, float InDeltaTime);
+	bool CanImport()  const;
+	FReply OnResetToDefaultClick();
+	FText GetImportTypeDisplayText() const;
+	void SetDefaultGridAssignment();
+	TSharedRef<SWidget> CreateGridSelector(TSharedPtr<SComboBox<TSharedPtr<FOpenVDBGridInfo>>>& ComboBox, TSharedPtr<SCheckBox>& CheckBox, FOpenVDBImportChannel& Channel, const FText& Label);
+};
\ No newline at end of file
diff --git a/Engine/Source/Editor/SparseVolumeTexture/Private/SparseVolumeTextureFactory.cpp b/Engine/Source/Editor/SparseVolumeTexture/Private/SparseVolumeTextureFactory.cpp
--- a/Engine/Source/Editor/SparseVolumeTexture/Private/SparseVolumeTextureFactory.cpp
+++ b/Engine/Source/Editor/SparseVolumeTexture/Private/SparseVolumeTextureFactory.cpp
@@ -16,6 +16,11 @@
 
 #include "Editor.h"
 
+#include "OpenVDBImportOptions.h"
+#include "OpenVDBImportWindow.h"
+#include "HAL/PlatformApplicationMisc.h"
+#include "Interfaces/IMainFrameModule.h"
+
 #define LOCTEXT_NAMESPACE "USparseVolumeTextureFactory"
 
 DEFINE_LOG_CATEGORY_STATIC(LogSparseVolumeTextureFactory, Log, All);
@@ -215,23 +220,89 @@ UObject* USparseVolumeTextureFactory::FactoryCreateFile(UClass* InClass, UObject
 		FName NewName(InName.ToString() + TEXT("VDB"));
 		UStaticSparseVolumeTexture* StaticSVTexture = NewObject<UStaticSparseVolumeTexture>(InParent, UStaticSparseVolumeTexture::StaticClass(), NewName, Flags);
 
+		// Load file and get info about each contained grid
 		TArray<uint8> LoadedFile;
-		if (!FFileHelper::LoadFileToArray(LoadedFile, *Filename))
+		TArray<TSharedPtr<FOpenVDBGridInfo>> GridInfoPtrs;
 		{
-			UE_LOG(LogSparseVolumeTextureFactory, Error, TEXT("OpenVDB file could not be opened: %s"), *Filename);
-			return nullptr;
+			if (!FFileHelper::LoadFileToArray(LoadedFile, *Filename))
+			{
+				UE_LOG(LogSparseVolumeTextureFactory, Error, TEXT("OpenVDB file could not be opened: %s"), *Filename);
+				return nullptr;
+			}
+
+			TArray<FOpenVDBGridInfo> GridInfo;
+			if (!GetOpenVDBGridInfo(LoadedFile, Filename, &GridInfo))
+			{
+				UE_LOG(LogSparseVolumeTextureFactory, Error, TEXT("OpenVDB file contains no suitable density grids: %s"), *Filename);
+				return nullptr;
+			}
+
+			// Convert to array of TSharedPtr. SComboBox requires this.
+			GridInfoPtrs.Empty(GridInfo.Num());
+			for (auto& Grid : GridInfo)
+			{
+				GridInfoPtrs.Add(MakeShared<FOpenVDBGridInfo>(Grid));
+			}
+		}
+
+		// Show import dialog
+		FOpenVDBImportOptions ImportOptions{};
+		{
+			TSharedPtr<SWindow> ParentWindow;
+
+			if (FModuleManager::Get().IsModuleLoaded("MainFrame"))
+			{
+				IMainFrameModule& MainFrame = FModuleManager::LoadModuleChecked<IMainFrameModule>("MainFrame");
+				ParentWindow = MainFrame.GetParentWindow();
+			}
+
+			// Compute centered window position based on max window size, which include when all categories are expanded
+			const float ImportWindowWidth = 450.0f;
+			const float ImportWindowHeight = 750.0f;
+			FVector2D ImportWindowSize = FVector2D(ImportWindowWidth, ImportWindowHeight); // Max window size it can get based on current slate
+
+
+			FSlateRect WorkAreaRect = FSlateApplicationBase::Get().GetPreferredWorkArea();
+			FVector2D DisplayTopLeft(WorkAreaRect.Left, WorkAreaRect.Top);
+			FVector2D DisplaySize(WorkAreaRect.Right - WorkAreaRect.Left, WorkAreaRect.Bottom - WorkAreaRect.Top);
+
+			float ScaleFactor = FPlatformApplicationMisc::GetDPIScaleFactorAtPoint(DisplayTopLeft.X, DisplayTopLeft.Y);
+			ImportWindowSize *= ScaleFactor;
+
+			FVector2D WindowPosition = (DisplayTopLeft + (DisplaySize - ImportWindowSize) / 2.0f) / ScaleFactor;
+
+			TSharedRef<SWindow> Window = SNew(SWindow)
+				.Title(NSLOCTEXT("UnrealEd", "OpenVDBImportOptionsTitle", "OpenVDB Import Options"))
+				.SizingRule(ESizingRule::Autosized)
+				.AutoCenter(EAutoCenter::None)
+				.ClientSize(ImportWindowSize)
+				.ScreenPosition(WindowPosition);
+
+			TSharedPtr<SOpenVDBImportWindow> OpenVDBOptionWindow;
+			Window->SetContent
+			(
+				SAssignNew(OpenVDBOptionWindow, SOpenVDBImportWindow)
+				.ImportOptions(&ImportOptions)
+				.OpenVDBGridInfo(&GridInfoPtrs)
+				.WidgetWindow(Window)
+				.FullPath(FText::FromString(Filename))
+				.MaxWindowHeight(ImportWindowHeight)
+				.MaxWindowWidth(ImportWindowWidth)
+			);
+
+			FSlateApplication::Get().AddModalWindow(Window, ParentWindow, false);
+
+			if (!OpenVDBOptionWindow->ShouldImport())
+			{
+				bOutOperationCanceled = true;
+				return nullptr;
+			}
 		}
 
 		FScopedSlowTask ImportTask(1.0f, LOCTEXT("ImportingVDBStatic", "Importing static OpenVDB"));
 		ImportTask.MakeDialog(true);
 
-		uint32 DensityGridIndex = 0;
-		FOpenVDBData OpenVDBData;
-		if (!FindDensityGridIndex(LoadedFile, Filename, &DensityGridIndex, &OpenVDBData))
-		{
-			UE_LOG(LogSparseVolumeTextureFactory, Error, TEXT("OpenVDB file contains no suitable density grid: %s"), *Filename);
-			return nullptr;
-		}
+		uint32 DensityGridIndex = ImportOptions.Density.Index;
 
 		FSparseVolumeRawSource SparseVolumeRawSource{};
 		SparseVolumeRawSource.DensityGridIndex = DensityGridIndex;
diff --git a/Engine/Source/Editor/SparseVolumeTexture/Private/SparseVolumeTextureOpenVDBUtility.cpp b/Engine/Source/Editor/SparseVolumeTexture/Private/SparseVolumeTextureOpenVDBUtility.cpp
--- a/Engine/Source/Editor/SparseVolumeTexture/Private/SparseVolumeTextureOpenVDBUtility.cpp
+++ b/Engine/Source/Editor/SparseVolumeTexture/Private/SparseVolumeTextureOpenVDBUtility.cpp
@@ -133,6 +133,65 @@ bool FindDensityGridIndex(TArray<uint8>& SourceFile, const FString& Filename, ui
 	return false;
 }
 
+bool GetOpenVDBGridInfo(TArray<uint8>& SourceFile, const FString& Filename, TArray<FOpenVDBGridInfo>* OutGridInfo)
+{
+#if OPENVDB_AVAILABLE
+	FArrayUint8StreamBuf StreamBuf(SourceFile);
+	std::istream IStream(&StreamBuf);
+	openvdb::io::Stream Stream(IStream, false /*delayLoad*/);
+
+	openvdb::GridPtrVecPtr Grids = Stream.getGrids();
+
+	if (!Grids)
+	{
+		UE_LOG(LogSparseVolumeTextureOpenVDBUtility, Warning, TEXT("OpenVDB file contains no grids: %s"), *Filename);
+		return false;
+	}
+
+	OutGridInfo->Empty(Grids->size());
+
+	uint32 GridIndex = 0;
+	for (openvdb::GridBase::Ptr& Grid : *Grids)
+	{
+		FOpenVDBGridInfo GridInfo;
+		GridInfo.Index = GridIndex;
+		++GridIndex;
+
+		// Figure out the type/format of the grid
+		bool bGridTypeSupported = false;
+		if (Grid->isType<openvdb::FloatGrid>())
+		{
+			GridInfo.Format = EOpenVDBGridFormat::Float;
+			bGridTypeSupported = true;
+		}
+		if (!bGridTypeSupported)
+		{
+			continue;
+		}
+
+		GridInfo.Name = Grid->getName().c_str();
+
+		FStringFormatOrderedArguments FormatArgs;
+		FormatArgs.Add(GridInfo.Name);
+		FormatArgs.Add(OpenVDBGridFormatToString(GridInfo.Format));
+
+		GridInfo.DisplayString = FString::Format(TEXT("Name: {0}, Format: {1}"), FormatArgs);
+
+		OutGridInfo->Add(MoveTemp(GridInfo));
+	}
+
+	if (OutGridInfo->IsEmpty())
+	{
+		UE_LOG(LogSparseVolumeTextureOpenVDBUtility, Warning, TEXT("OpenVDB file contains no grids of supported type/format: %s"), *Filename);
+		return false;
+	}
+
+	return true;
+
+#endif // OPENVDB_AVAILABLE
+	return false;
+}
+
 bool ConvertOpenVDBToSparseVolumeTexture(
 	TArray<uint8>& SourceFile,
 	uint32 GridIndex,
@@ -299,7 +358,7 @@ bool ConvertOpenVDBToSparseVolumeTexture(
 			[
 				PageCoordToSplat.Z * Header.PageTableVolumeResolution.X * Header.PageTableVolumeResolution.Y +
 				PageCoordToSplat.Y * Header.PageTableVolumeResolution.X +
-			PageCoordToSplat.X
+				PageCoordToSplat.X
 			] = DestinationTileCoord32bit;
 
 		// Now copy the tile data
@@ -313,7 +372,7 @@ bool ConvertOpenVDBToSparseVolumeTexture(
 						[
 							(DestinationTileCoord.Z * SPARSE_VOLUME_TILE_RES + z) * Header.TileDataVolumeResolution.X * Header.TileDataVolumeResolution.Y +
 							(DestinationTileCoord.Y * SPARSE_VOLUME_TILE_RES + y) * Header.TileDataVolumeResolution.X +
-						(DestinationTileCoord.X * SPARSE_VOLUME_TILE_RES + x)
+							(DestinationTileCoord.X * SPARSE_VOLUME_TILE_RES + x)
 						] = TileDataToSplat[z * SPARSE_VOLUME_TILE_RES * SPARSE_VOLUME_TILE_RES + y * SPARSE_VOLUME_TILE_RES + x];
 				}
 			}
@@ -329,4 +388,15 @@ bool ConvertOpenVDBToSparseVolumeTexture(
 #endif // OPENVDB_AVAILABLE
 }
 
+const TCHAR* OpenVDBGridFormatToString(EOpenVDBGridFormat Format)
+{
+	switch (Format)
+	{
+	case EOpenVDBGridFormat::Float:
+		return TEXT("Float");
+	default:
+		return TEXT("Unknown");
+	}
+}
+
 #endif // WITH_EDITOR
\ No newline at end of file
diff --git a/Engine/Source/Editor/SparseVolumeTexture/Private/SparseVolumeTextureOpenVDBUtility.h b/Engine/Source/Editor/SparseVolumeTexture/Private/SparseVolumeTextureOpenVDBUtility.h
--- a/Engine/Source/Editor/SparseVolumeTexture/Private/SparseVolumeTextureOpenVDBUtility.h
+++ b/Engine/Source/Editor/SparseVolumeTexture/Private/SparseVolumeTextureOpenVDBUtility.h
@@ -16,10 +16,25 @@ struct FOpenVDBData
 	bool bHasUniformVoxels;
 };
 
+enum class EOpenVDBGridFormat : uint8
+{
+	Float = 0,
+};
+
+struct FOpenVDBGridInfo
+{
+	uint32 Index;
+	EOpenVDBGridFormat Format;
+	FString Name;
+	FString DisplayString; // Contains both Name and Format
+};
+
 bool IsOpenVDBDataValid(FOpenVDBData& OpenVDBData, const FString& Filename);
 
 bool FindDensityGridIndex(TArray<uint8>& SourceFile, const FString& Filename, uint32* OutGridIndex, FOpenVDBData* OutOVDBData);
 
+bool GetOpenVDBGridInfo(TArray<uint8>& SourceFile, const FString& Filename, TArray<FOpenVDBGridInfo>* OutGridInfo);
+
 bool ConvertOpenVDBToSparseVolumeTexture(
 	TArray<uint8>& SourceFile,
 	uint32 GridIndex,
@@ -30,4 +45,6 @@ bool ConvertOpenVDBToSparseVolumeTexture(
 	FVector ActiveMin,
 	FVector ActiveMax);
 
+const TCHAR* OpenVDBGridFormatToString(EOpenVDBGridFormat Format);
+
 #endif // WITH_EDITOR
\ No newline at end of file
diff --git a/Engine/Source/Editor/SparseVolumeTexture/SparseVolumeTexture.Build.cs b/Engine/Source/Editor/SparseVolumeTexture/SparseVolumeTexture.Build.cs
--- a/Engine/Source/Editor/SparseVolumeTexture/SparseVolumeTexture.Build.cs
+++ b/Engine/Source/Editor/SparseVolumeTexture/SparseVolumeTexture.Build.cs
@@ -17,6 +17,10 @@ public class SparseVolumeTexture : ModuleRules
 				"Engine",
 				"Renderer",
 				"RenderCore",
+				"Slate",
+				"SlateCore",
+				"ApplicationCore",
+				"ToolWidgets",
 			}
 		);
 
