Fixed missing copyright boilerplate.

#rnx

diff --git a/Engine/Source/Editor/SparseVolumeTexture/Private/OpenVDBImportWindow.cpp b/Engine/Source/Editor/SparseVolumeTexture/Private/OpenVDBImportWindow.cpp
--- a/Engine/Source/Editor/SparseVolumeTexture/Private/OpenVDBImportWindow.cpp
+++ b/Engine/Source/Editor/SparseVolumeTexture/Private/OpenVDBImportWindow.cpp
@@ -1,3 +1,5 @@
+// Copyright Epic Games, Inc. All Rights Reserved.
+
 #include "OpenVDBImportWindow.h"
 
 #include "Widgets/Input/SButton.h"
