Reworked OpenVDB import dialog to allow for combining source grids into different channels of the resulting page atlas. Also now supports three different formats (32bit float, 16bit float, 8bit unorm).

#rnx

diff --git a/Engine/Shaders/Private/SparseVolumeTexture/VisualizeSparseVolumeTexture.usf b/Engine/Shaders/Private/SparseVolumeTexture/VisualizeSparseVolumeTexture.usf
--- a/Engine/Shaders/Private/SparseVolumeTexture/VisualizeSparseVolumeTexture.usf
+++ b/Engine/Shaders/Private/SparseVolumeTexture/VisualizeSparseVolumeTexture.usf
@@ -45,6 +45,7 @@ float3 VolumeBoundMaxWorld;
 float4 WorldToLocal0;
 float4 WorldToLocal1;
 float4 WorldToLocal2;
+uint ComponentToVisualize;
 
 
 Texture3D<float4> SparseVolumeTexture;
@@ -180,7 +181,7 @@ void VisualizeSparseVolumeTexturePS(
 
 				const float4 VoxelData = SparseVolumeTexture.Load(int4(VoxelCoord, 0));
 
-				const float Density = VoxelData.x;
+				const float Density = VoxelData[ComponentToVisualize];
 				if (Density > 0.0f)
 
 				{
diff --git a/Engine/Source/Editor/SparseVolumeTexture/Private/OpenVDBImportOptions.h b/Engine/Source/Editor/SparseVolumeTexture/Private/OpenVDBImportOptions.h
deleted file mode 100644
diff --git a/Engine/Source/Editor/SparseVolumeTexture/Private/OpenVDBImportWindow.cpp b/Engine/Source/Editor/SparseVolumeTexture/Private/OpenVDBImportWindow.cpp
--- a/Engine/Source/Editor/SparseVolumeTexture/Private/OpenVDBImportWindow.cpp
+++ b/Engine/Source/Editor/SparseVolumeTexture/Private/OpenVDBImportWindow.cpp
@@ -13,7 +13,7 @@
 #include "IDocumentation.h"
 #include "Editor.h"
 #include "SparseVolumeTextureOpenVDBUtility.h"
-#include "OpenVDBImportOptions.h"
+#include "SparseVolumeTexture/SparseVolumeTexture.h"
 
 #define LOCTEXT_NAMESPACE "SOpenVDBImportWindow"
 
@@ -22,10 +22,26 @@ static FText GetGridComboBoxItemText(TSharedPtr<FOpenVDBGridInfo> InItem)
 	return InItem ? FText::FromString(InItem->DisplayString) : LOCTEXT("NoneGrid", "<None>");
 };
 
+static FText GetFormatComboBoxItemText(TSharedPtr<ESparseVolumePackedDataFormat> InItem)
+{
+	const TCHAR* FormatStr = TEXT("<None>");
+	if (InItem)
+	{
+		switch (*InItem)
+		{
+		case ESparseVolumePackedDataFormat::Unorm8: FormatStr = TEXT("8bit unorm"); break;
+		case ESparseVolumePackedDataFormat::Float16: FormatStr = TEXT("16bit float"); break;
+		case ESparseVolumePackedDataFormat::Float32: FormatStr = TEXT("32bit float"); break;
+		}
+	}
+	return FText::FromString(FormatStr);
+}
+
 void SOpenVDBImportWindow::Construct(const FArguments& InArgs)
 {
-	ImportOptions = InArgs._ImportOptions;
+	PackedDataA = InArgs._PackedDataA;
 	OpenVDBGridInfo = InArgs._OpenVDBGridInfo;
+	OpenVDBSupportedTargetFormats = InArgs._OpenVDBSupportedTargetFormats;
 	WidgetWindow = InArgs._WidgetWindow;
 
 	TSharedPtr<SBox> ImportTypeDisplay;
@@ -109,7 +125,11 @@ void SOpenVDBImportWindow::Construct(const FArguments& InArgs)
 		.AutoHeight()
 		.Padding(2)
 		[
-			CreateGridSelector(DensityGridComboBox, DensityGridCheckBox, ImportOptions->Density, LOCTEXT("DensityGridIndex", "Density"))
+			SAssignNew(PackedDataAConfigurator, SOpenVDBPackedDataConfigurator)
+			.PackedData(PackedDataA)
+			.OpenVDBGridInfo(OpenVDBGridInfo)
+			.OpenVDBSupportedTargetFormats(OpenVDBSupportedTargetFormats)
+			.PackedDataName(LOCTEXT("OpenVDBImportWindow_PackedDataA", "Packed Data A"))
 		]
 	);
 
@@ -190,7 +210,16 @@ EActiveTimerReturnType SOpenVDBImportWindow::SetFocusPostConstruct(double InCurr
 
 bool SOpenVDBImportWindow::CanImport() const
 {
-	return DensityGridCheckBox->IsChecked();
+	const FUintVector4& GridIndices = PackedDataA->SourceGridIndex;
+	const FUintVector4& ComponentIndices = PackedDataA->SourceComponentIndex;
+	if ((GridIndices.X != INDEX_NONE && ComponentIndices.X != INDEX_NONE)
+		|| (GridIndices.Y != INDEX_NONE && ComponentIndices.Y != INDEX_NONE)
+		|| (GridIndices.Z != INDEX_NONE && ComponentIndices.Z != INDEX_NONE)
+		|| (GridIndices.W != INDEX_NONE && ComponentIndices.W != INDEX_NONE))
+	{
+		return true;
+	}
+	return false;
 }
 
 FReply SOpenVDBImportWindow::OnResetToDefaultClick()
@@ -208,122 +237,252 @@ void SOpenVDBImportWindow::SetDefaultGridAssignment()
 {
 	check(OpenVDBGridInfo);
 
-	auto SetDefaultGrid = [](const TArray<TSharedPtr<FOpenVDBGridInfo>>& Grids, const FString& Name, EOpenVDBGridFormat Format, 
-		TSharedPtr<SComboBox<TSharedPtr<FOpenVDBGridInfo>>> ComboBox, TSharedPtr<SCheckBox> CheckBox, FOpenVDBImportChannel& Channel)
-	{
-		TSharedPtr<FOpenVDBGridInfo> SelectedGrid = nullptr;
-		// Search by name first...
-		for (auto& Grid : Grids)
-		{
-			if (Grid->Name == Name)
-			{
-				SelectedGrid = Grid;
-				break;
-			}
-		}
+	PackedDataA->Format = ESparseVolumePackedDataFormat::Float32;
+	PackedDataA->SourceGridIndex = FUintVector4(INDEX_NONE);
+	PackedDataA->SourceComponentIndex = FUintVector4(INDEX_NONE);
+	PackedDataA->bRescaleInputForUnorm = false;
+
+	const TCHAR* SearchNames[] = { TEXT("density"), TEXT("heat"), TEXT("temperature"), TEXT("motion") };
 
-		// ...and if that fails, just take the first grid with matching format.
-		if (!SelectedGrid)
+	uint32 CurrentOutputComponent = 0;
+	for (const TCHAR* SearchName : SearchNames)
+	{
+		for (uint32 InputComponent = 0; InputComponent < 4; ++InputComponent)
 		{
-			for (auto& Grid : Grids)
+			for (const TSharedPtr<FOpenVDBGridInfo>& Grid : *OpenVDBGridInfo)
 			{
-				if (Grid->Format == Format)
+				if (Grid->Name == SearchName && Grid->ComponentIndex == InputComponent)
 				{
-					SelectedGrid = Grid;
+					PackedDataA->SourceGridIndex[CurrentOutputComponent] = Grid->Index;
+					PackedDataA->SourceComponentIndex[CurrentOutputComponent] = Grid->ComponentIndex;
+					++CurrentOutputComponent;
 					break;
 				}
 			}
 		}
+	}
 
-		if (SelectedGrid)
-		{
-			const bool bImport = SelectedGrid->Name == Name; // disable importing this grid if the name doesn't match
-			ComboBox->SetSelectedItem(SelectedGrid);
-			CheckBox->SetIsChecked(bImport);
-			Channel.Name = SelectedGrid->Name;
-			Channel.Index = SelectedGrid->Index;
-			Channel.bImport = bImport;
-		}
-		else
-		{
-			ComboBox->ClearSelection();
-			CheckBox->SetIsChecked(false);
-			Channel.Name.Reset();
-			Channel.Index = INDEX_NONE;
-			Channel.bImport = false;
-		}
-	};
-
-	SetDefaultGrid(*OpenVDBGridInfo, TEXT("density"), EOpenVDBGridFormat::Float, DensityGridComboBox, DensityGridCheckBox, ImportOptions->Density);
+	PackedDataAConfigurator->RefreshUIFromData();
 }
 
-TSharedRef<SWidget> SOpenVDBImportWindow::CreateGridSelector(TSharedPtr<SComboBox<TSharedPtr<FOpenVDBGridInfo>>>& ComboBox, TSharedPtr<SCheckBox>& CheckBox, FOpenVDBImportChannel& Channel, const FText& Label)
+void SOpenVDBComponentPicker::Construct(const FArguments& InArgs)
 {
-	return
-		SNew(SHorizontalBox)
-		+ SHorizontalBox::Slot()
-		.VAlign(VAlign_Center)
-		.AutoWidth()
-		.Padding(2.0f)
-		[
-			SAssignNew(CheckBox, SCheckBox)
-			.IsChecked(true)
-			.OnCheckStateChanged_Lambda([&Channel](ECheckBoxState CheckBoxState)
-			{
-				Channel.bImport = CheckBoxState == ECheckBoxState::Checked;
-			})
-		]
+	PackedData = InArgs._PackedData;
+	ComponentIndex = InArgs._ComponentIndex;
+	OpenVDBGridInfo = InArgs._OpenVDBGridInfo;
+	
+	check(ComponentIndex < 4);
+	const TCHAR* ComponentLabels[] = { TEXT("X"), TEXT("Y"), TEXT("Z"), TEXT("W") };
 
-		+ SHorizontalBox::Slot()
-		.VAlign(VAlign_Center)
-		.FillWidth(1.0f)
-		.Padding(2.0f)
+	this->ChildSlot
 		[
-			SNew(STextBlock)
-			.Text(Label)
-		]
+			SNew(SHorizontalBox)
+			
+			+ SHorizontalBox::Slot()
+			.VAlign(VAlign_Center)
+			.FillWidth(1.0f)
+			.Padding(2.0f)
+			[
+				SNew(STextBlock)
+				.Text(FText::FromString(ComponentLabels[ComponentIndex]))
+			]
 
-		+ SHorizontalBox::Slot()
-		.FillWidth(1.0f)
-		.VAlign(VAlign_Center)
-		.Padding(2.0f)
-		[
-			SNew(SBox)
-			.WidthOverride(300.0f)
+			+ SHorizontalBox::Slot()
+			.FillWidth(1.0f)
+			.VAlign(VAlign_Center)
+			.Padding(2.0f)
 			[
-				SAssignNew(ComboBox, SComboBox<TSharedPtr<FOpenVDBGridInfo>>)
-				.OptionsSource(OpenVDBGridInfo)
-				.IsEnabled_Lambda([&CheckBox]()
-				{
-					return CheckBox->IsChecked();
-				})
-				.OnGenerateWidget_Lambda([](TSharedPtr<FOpenVDBGridInfo> InItem)
-				{
-					return SNew(STextBlock)
-					.Text(GetGridComboBoxItemText(InItem));
-				})
-				.OnSelectionChanged_Lambda([&Channel](TSharedPtr<FOpenVDBGridInfo> InItem, ESelectInfo::Type)
-				{
-					if (InItem)
+				SNew(SBox)
+				.WidthOverride(300.0f)
+				[
+					SAssignNew(GridComboBox, SComboBox<TSharedPtr<FOpenVDBGridInfo>>)
+					.OptionsSource(OpenVDBGridInfo)
+					.OnGenerateWidget_Lambda([](TSharedPtr<FOpenVDBGridInfo> InItem)
 					{
-						Channel.Name = InItem->Name;
-						Channel.Index = InItem->Index;
-					}
-					else
+						return SNew(STextBlock)
+						.Text(GetGridComboBoxItemText(InItem));
+					})
+					.OnSelectionChanged_Lambda([this](TSharedPtr<FOpenVDBGridInfo> InItem, ESelectInfo::Type)
 					{
-						Channel.Name.Reset();
-						Channel.Index = INDEX_NONE;
-					}
-				})
+						if (InItem)
+						{
+							PackedData->SourceGridIndex[ComponentIndex] = InItem->Index;
+							PackedData->SourceComponentIndex[ComponentIndex] = InItem->ComponentIndex;
+						}
+						else
+						{
+							PackedData->SourceGridIndex[ComponentIndex] = INDEX_NONE;
+							PackedData->SourceComponentIndex[ComponentIndex] = INDEX_NONE;
+						}
+					})
+					[
+						SNew(STextBlock)
+						.Text_Lambda([this]()
+						{
+							return GetGridComboBoxItemText(GridComboBox->GetSelectedItem());
+						})
+					]
+				]
+			]
+		];
+}
+
+void SOpenVDBComponentPicker::RefreshUIFromData()
+{
+	for (const TSharedPtr<FOpenVDBGridInfo>& Grid : *OpenVDBGridInfo)
+	{
+		printf("%i\n", Grid->Index);
+		if (Grid->Index == PackedData->SourceGridIndex[ComponentIndex] && Grid->ComponentIndex == PackedData->SourceComponentIndex[ComponentIndex])
+		{
+			GridComboBox->SetSelectedItem(Grid);
+			break;
+		}
+	}
+}
+
+void SOpenVDBPackedDataConfigurator::Construct(const FArguments& InArgs)
+{
+	PackedData = InArgs._PackedData;
+	OpenVDBSupportedTargetFormats = InArgs._OpenVDBSupportedTargetFormats;
+
+	for (uint32 ComponentIndex = 0; ComponentIndex < 4; ++ComponentIndex)
+	{
+		ComponentPickers[ComponentIndex] =
+			SNew(SOpenVDBComponentPicker)
+			.PackedData(PackedData)
+			.ComponentIndex(ComponentIndex)
+			.OpenVDBGridInfo(InArgs._OpenVDBGridInfo);
+	}
+
+	this->ChildSlot
+		[
+			SNew(SVerticalBox)
+			
+			+ SVerticalBox::Slot()
+			.AutoHeight()
+			.Padding(2)
+			[
+				SNew(SHorizontalBox)
+
+				+ SHorizontalBox::Slot()
+				.VAlign(VAlign_Center)
+				.FillWidth(1.0f)
+				.Padding(2.0f)
 				[
 					SNew(STextBlock)
-					.Text_Lambda([&ComboBox]()
+					.Text(InArgs._PackedDataName)
+				]
+
+				+ SHorizontalBox::Slot()
+				.FillWidth(1.0f)
+				.VAlign(VAlign_Center)
+				.Padding(2.0f)
+				[
+					SNew(SBox)
+					.WidthOverride(50.0f)
+					[
+						SAssignNew(FormatComboBox, SComboBox<TSharedPtr<ESparseVolumePackedDataFormat>>)
+						.OptionsSource(OpenVDBSupportedTargetFormats)
+						.OnGenerateWidget_Lambda([](TSharedPtr<ESparseVolumePackedDataFormat> InItem)
+						{
+							return SNew(STextBlock)
+							.Text(GetFormatComboBoxItemText(InItem));
+						})
+						.OnSelectionChanged_Lambda([this](TSharedPtr<ESparseVolumePackedDataFormat> InItem, ESelectInfo::Type)
+						{
+							PackedData->Format = InItem ? *InItem : ESparseVolumePackedDataFormat::Float32;
+						})
+						[
+							SNew(STextBlock)
+							.Text_Lambda([this]()
+							{
+								return GetFormatComboBoxItemText(FormatComboBox->GetSelectedItem());
+							})
+						]
+					]
+				]
+				+ SHorizontalBox::Slot()
+				.VAlign(VAlign_Center)
+				.HAlign(HAlign_Right)
+				.AutoWidth()
+				.Padding(2.0f)
+				[
+					SNew(STextBlock)
+					.Text(LOCTEXT("UnormRescaleCheckBoxLabel", "Unorm Rescale"))
+				]
+				+ SHorizontalBox::Slot()
+				.VAlign(VAlign_Center)
+				.HAlign(HAlign_Right)
+				.AutoWidth()
+				.Padding(2.0f)
+				[
+					SAssignNew(RescaleUnormCheckBox, SCheckBox)
+					.OnCheckStateChanged_Lambda([this](ECheckBoxState CheckBoxState)
+					{
+							PackedData->bRescaleInputForUnorm = CheckBoxState == ECheckBoxState::Checked;
+					})
+					.IsEnabled_Lambda([this]()
 					{
-						return GetGridComboBoxItemText(ComboBox->GetSelectedItem());
+						return PackedData->Format == ESparseVolumePackedDataFormat::Unorm8;
 					})
+					.ToolTipText(LOCTEXT("UnormRescaleCheckBoxTooltip", "Rescales input values for unorm formats into the [0-1] range instead of clamping values outside this range."))
+					.IsChecked(false)
+				]
+			]
+
+			+ SVerticalBox::Slot()
+			.AutoHeight()
+			.Padding(2)
+			[
+				SNew(SBorder)
+				.Padding(FMargin(3))
+				.BorderImage(FAppStyle::GetBrush("ToolPanel.GroupBorder"))
+				[
+					SNew(SVerticalBox)
+					+ SVerticalBox::Slot()
+					.AutoHeight()
+					.Padding(2)
+					[
+						ComponentPickers[0]->AsShared()
+					]
+					+ SVerticalBox::Slot()
+					.AutoHeight()
+					.Padding(2)
+					[
+						ComponentPickers[1]->AsShared()
+					]
+					+ SVerticalBox::Slot()
+					.AutoHeight()
+					.Padding(2)
+					[
+						ComponentPickers[2]->AsShared()
+					]
+					+ SVerticalBox::Slot()
+					.AutoHeight()
+					.Padding(2)
+					[
+						ComponentPickers[3]->AsShared()
+					]
 				]
 			]
 		];
 }
 
+void SOpenVDBPackedDataConfigurator::RefreshUIFromData()
+{
+	for (auto& Format : *OpenVDBSupportedTargetFormats)
+	{
+		if (*Format == PackedData->Format)
+		{
+			FormatComboBox->SetSelectedItem(Format);
+			break;
+		}
+	}
+	for (uint32 i = 0; i < 4; ++i)
+	{
+		ComponentPickers[i]->RefreshUIFromData();
+	}
+	RescaleUnormCheckBox->SetIsChecked(PackedData->bRescaleInputForUnorm);
+}
+
 #undef LOCTEXT_NAMESPACE
diff --git a/Engine/Source/Editor/SparseVolumeTexture/Private/OpenVDBImportWindow.h b/Engine/Source/Editor/SparseVolumeTexture/Private/OpenVDBImportWindow.h
--- a/Engine/Source/Editor/SparseVolumeTexture/Private/OpenVDBImportWindow.h
+++ b/Engine/Source/Editor/SparseVolumeTexture/Private/OpenVDBImportWindow.h
@@ -8,23 +8,78 @@
 #include "Widgets/Input/SCheckBox.h"
 
 struct FOpenVDBGridInfo;
-struct FOpenVDBImportOptions;
-struct FOpenVDBImportChannel;
+struct ENGINE_API FSparseVolumeRawSourcePackedData;
+enum class ESparseVolumePackedDataFormat : uint8;
+
+class SOpenVDBComponentPicker : public SCompoundWidget
+{
+public:
+	SLATE_BEGIN_ARGS(SOpenVDBComponentPicker)
+		: _PackedData()
+		, _ComponentIndex()
+		, _OpenVDBGridInfo()
+	{}
+
+		SLATE_ARGUMENT(FSparseVolumeRawSourcePackedData*, PackedData)
+		SLATE_ARGUMENT(uint32, ComponentIndex)
+		SLATE_ARGUMENT(TArray<TSharedPtr<FOpenVDBGridInfo>>*, OpenVDBGridInfo)
+	SLATE_END_ARGS()
+
+public:
+	void Construct(const FArguments& InArgs);
+	void RefreshUIFromData();
+
+private:
+	FSparseVolumeRawSourcePackedData* PackedData;
+	uint32 ComponentIndex;
+	TArray<TSharedPtr<FOpenVDBGridInfo>>* OpenVDBGridInfo;
+	TSharedPtr<SComboBox<TSharedPtr<FOpenVDBGridInfo>>>	GridComboBox;
+};
+
+class SOpenVDBPackedDataConfigurator : public SCompoundWidget
+{
+public:
+	SLATE_BEGIN_ARGS(SOpenVDBPackedDataConfigurator)
+		: _PackedData()
+		, _OpenVDBGridInfo()
+		, _OpenVDBSupportedTargetFormats()
+		, _PackedDataName()
+	{}
+
+		SLATE_ARGUMENT(FSparseVolumeRawSourcePackedData*, PackedData)
+		SLATE_ARGUMENT(TArray<TSharedPtr<FOpenVDBGridInfo>>*, OpenVDBGridInfo)
+		SLATE_ARGUMENT(TArray<TSharedPtr<ESparseVolumePackedDataFormat>>*, OpenVDBSupportedTargetFormats)
+		SLATE_ARGUMENT(FText, PackedDataName)
+	SLATE_END_ARGS()
+
+public:
+	void Construct(const FArguments& InArgs);
+	void RefreshUIFromData();
+
+private:
+	FSparseVolumeRawSourcePackedData* PackedData;
+	TSharedPtr<SOpenVDBComponentPicker> ComponentPickers[4];
+	TArray<TSharedPtr<ESparseVolumePackedDataFormat>>* OpenVDBSupportedTargetFormats;
+	TSharedPtr<SComboBox<TSharedPtr<ESparseVolumePackedDataFormat>>>	FormatComboBox;
+	TSharedPtr<SCheckBox> RescaleUnormCheckBox;
+};
 
 class SOpenVDBImportWindow : public SCompoundWidget
 {
 public:
 	SLATE_BEGIN_ARGS(SOpenVDBImportWindow)
-		: _ImportOptions()
-		,_OpenVDBGridInfo()
+		: _PackedDataA()
+		, _OpenVDBGridInfo()
+		, _OpenVDBSupportedTargetFormats()
 		, _WidgetWindow()
 		, _FullPath()
 		, _MaxWindowHeight(0.0f)
 		, _MaxWindowWidth(0.0f)
 	{}
 
-		SLATE_ARGUMENT(FOpenVDBImportOptions*, ImportOptions)
+		SLATE_ARGUMENT(FSparseVolumeRawSourcePackedData*, PackedDataA)
 		SLATE_ARGUMENT(TArray<TSharedPtr<FOpenVDBGridInfo>>*, OpenVDBGridInfo)
+		SLATE_ARGUMENT(TArray<TSharedPtr<ESparseVolumePackedDataFormat>>*, OpenVDBSupportedTargetFormats)
 		SLATE_ARGUMENT(TSharedPtr<SWindow>, WidgetWindow)
 		SLATE_ARGUMENT(FText, FullPath)
 		SLATE_ARGUMENT(float, MaxWindowHeight)
@@ -39,10 +94,10 @@ public:
 	bool ShouldImport() const;
 
 private:
-	FOpenVDBImportOptions*								ImportOptions;
+	FSparseVolumeRawSourcePackedData*					PackedDataA;
 	TArray<TSharedPtr<FOpenVDBGridInfo>>*				OpenVDBGridInfo;
-	TSharedPtr<SComboBox<TSharedPtr<FOpenVDBGridInfo>>>	DensityGridComboBox;
-	TSharedPtr<SCheckBox>								DensityGridCheckBox;
+	TArray<TSharedPtr<ESparseVolumePackedDataFormat>>*	OpenVDBSupportedTargetFormats;
+	TSharedPtr<SOpenVDBPackedDataConfigurator>			PackedDataAConfigurator;
 	TSharedPtr<SButton>									ImportButton;
 	TWeakPtr<SWindow>									WidgetWindow;
 	bool												bShouldImport;
@@ -52,5 +107,4 @@ private:
 	FReply OnResetToDefaultClick();
 	FText GetImportTypeDisplayText() const;
 	void SetDefaultGridAssignment();
-	TSharedRef<SWidget> CreateGridSelector(TSharedPtr<SComboBox<TSharedPtr<FOpenVDBGridInfo>>>& ComboBox, TSharedPtr<SCheckBox>& CheckBox, FOpenVDBImportChannel& Channel, const FText& Label);
 };
\ No newline at end of file
diff --git a/Engine/Source/Editor/SparseVolumeTexture/Private/SparseVolumeTextureFactory.cpp b/Engine/Source/Editor/SparseVolumeTexture/Private/SparseVolumeTextureFactory.cpp
--- a/Engine/Source/Editor/SparseVolumeTexture/Private/SparseVolumeTextureFactory.cpp
+++ b/Engine/Source/Editor/SparseVolumeTexture/Private/SparseVolumeTextureFactory.cpp
@@ -16,7 +16,6 @@
 
 #include "Editor.h"
 
-#include "OpenVDBImportOptions.h"
 #include "OpenVDBImportWindow.h"
 #include "HAL/PlatformApplicationMisc.h"
 #include "Interfaces/IMainFrameModule.h"
@@ -194,7 +193,9 @@ UObject* USparseVolumeTextureFactory::FactoryCreateFile(UClass* InClass, UObject
 				OpenVDBData.VolumeActiveDim.X, OpenVDBData.VolumeActiveDim.Y, OpenVDBData.VolumeActiveDim.Z);
 
 			FSparseVolumeRawSource SparseVolumeRawSource{};
-			SparseVolumeRawSource.DensityGridIndex = DensityGridIndex;
+			SparseVolumeRawSource.PackedDataA.Format = ESparseVolumePackedDataFormat::Unorm8;
+			SparseVolumeRawSource.PackedDataA.SourceGridIndex = FUintVector4(0, INDEX_NONE, INDEX_NONE, INDEX_NONE);
+			SparseVolumeRawSource.PackedDataA.SourceComponentIndex = FUintVector4(0, INDEX_NONE, INDEX_NONE, INDEX_NONE);
 			SparseVolumeRawSource.SourceAssetFile = MoveTemp(LoadedFile);
 			LoadedFile.Reset();
 
@@ -223,6 +224,12 @@ UObject* USparseVolumeTextureFactory::FactoryCreateFile(UClass* InClass, UObject
 		// Load file and get info about each contained grid
 		TArray<uint8> LoadedFile;
 		TArray<TSharedPtr<FOpenVDBGridInfo>> GridInfoPtrs;
+		TArray<TSharedPtr<ESparseVolumePackedDataFormat>> SupportedFormats =
+		{ 
+			MakeShared<ESparseVolumePackedDataFormat>(ESparseVolumePackedDataFormat::Float32),
+			MakeShared<ESparseVolumePackedDataFormat>(ESparseVolumePackedDataFormat::Float16),
+			MakeShared<ESparseVolumePackedDataFormat>(ESparseVolumePackedDataFormat::Unorm8)
+		};
 		{
 			if (!FFileHelper::LoadFileToArray(LoadedFile, *Filename))
 			{
@@ -238,7 +245,18 @@ UObject* USparseVolumeTextureFactory::FactoryCreateFile(UClass* InClass, UObject
 			}
 
 			// Convert to array of TSharedPtr. SComboBox requires this.
-			GridInfoPtrs.Empty(GridInfo.Num());
+			GridInfoPtrs.Empty(GridInfo.Num() + 1);
+			
+			// We need a <None> option to leave channels empty
+			FOpenVDBGridInfo NoneGridInfo;
+			NoneGridInfo.Index = INDEX_NONE;
+			NoneGridInfo.ComponentIndex = INDEX_NONE;
+			NoneGridInfo.Format = EOpenVDBGridFormat::Float;
+			NoneGridInfo.Name = TEXT("<None>");
+			NoneGridInfo.DisplayString = TEXT("<None>");
+			
+			GridInfoPtrs.Add(MakeShared<FOpenVDBGridInfo>(NoneGridInfo));
+			
 			for (auto& Grid : GridInfo)
 			{
 				GridInfoPtrs.Add(MakeShared<FOpenVDBGridInfo>(Grid));
@@ -246,7 +264,7 @@ UObject* USparseVolumeTextureFactory::FactoryCreateFile(UClass* InClass, UObject
 		}
 
 		// Show import dialog
-		FOpenVDBImportOptions ImportOptions{};
+		FSparseVolumeRawSourcePackedData PackedDataA{};
 		{
 			TSharedPtr<SWindow> ParentWindow;
 
@@ -282,8 +300,9 @@ UObject* USparseVolumeTextureFactory::FactoryCreateFile(UClass* InClass, UObject
 			Window->SetContent
 			(
 				SAssignNew(OpenVDBOptionWindow, SOpenVDBImportWindow)
-				.ImportOptions(&ImportOptions)
+				.PackedDataA(&PackedDataA)
 				.OpenVDBGridInfo(&GridInfoPtrs)
+				.OpenVDBSupportedTargetFormats(&SupportedFormats)
 				.WidgetWindow(Window)
 				.FullPath(FText::FromString(Filename))
 				.MaxWindowHeight(ImportWindowHeight)
@@ -302,10 +321,8 @@ UObject* USparseVolumeTextureFactory::FactoryCreateFile(UClass* InClass, UObject
 		FScopedSlowTask ImportTask(1.0f, LOCTEXT("ImportingVDBStatic", "Importing static OpenVDB"));
 		ImportTask.MakeDialog(true);
 
-		uint32 DensityGridIndex = ImportOptions.Density.Index;
-
 		FSparseVolumeRawSource SparseVolumeRawSource{};
-		SparseVolumeRawSource.DensityGridIndex = DensityGridIndex;
+		SparseVolumeRawSource.PackedDataA = PackedDataA;
 		SparseVolumeRawSource.SourceAssetFile = MoveTemp(LoadedFile);
 		LoadedFile.Reset();
 
diff --git a/Engine/Source/Editor/SparseVolumeTexture/Private/SparseVolumeTextureOpenVDBUtility.cpp b/Engine/Source/Editor/SparseVolumeTexture/Private/SparseVolumeTextureOpenVDBUtility.cpp
--- a/Engine/Source/Editor/SparseVolumeTexture/Private/SparseVolumeTextureOpenVDBUtility.cpp
+++ b/Engine/Source/Editor/SparseVolumeTexture/Private/SparseVolumeTextureOpenVDBUtility.cpp
@@ -148,36 +148,82 @@ bool GetOpenVDBGridInfo(TArray<uint8>& SourceFile, const FString& Filename, TArr
 		return false;
 	}
 
-	OutGridInfo->Empty(Grids->size());
+	OutGridInfo->Empty(Grids->size() * 2 /*assume an average of 2 channels per grid*/);
 
 	uint32 GridIndex = 0;
 	for (openvdb::GridBase::Ptr& Grid : *Grids)
 	{
 		FOpenVDBGridInfo GridInfo;
-		GridInfo.Index = GridIndex;
-		++GridIndex;
+		GridInfo.Index = GridIndex++;
 
 		// Figure out the type/format of the grid
-		bool bGridTypeSupported = false;
+		uint32 NumComponents = 0;
 		if (Grid->isType<openvdb::FloatGrid>())
 		{
+			NumComponents = 1;
+			GridInfo.Format = EOpenVDBGridFormat::Float;
+		}
+#if 0 // SV_TODO: actually support these additional types in ConvertOpenVDBToSparseVolumeTexture()
+		else if (Grid->isType<openvdb::Grid<openvdb::tree::Tree4<openvdb::Vec2f, 5, 4, 3>::Type>>())
+		{
+			NumComponents = 2;
+			GridInfo.Format = EOpenVDBGridFormat::Float;
+		}
+		else if (Grid->isType<openvdb::Vec3SGrid>())
+		{
+			NumComponents = 3;
+			GridInfo.Format = EOpenVDBGridFormat::Float;
+		}
+		else if (Grid->isType<openvdb::Grid<openvdb::tree::Tree4<openvdb::Vec4f, 5, 4, 3>::Type>>())
+		{
+			NumComponents = 4;
 			GridInfo.Format = EOpenVDBGridFormat::Float;
-			bGridTypeSupported = true;
 		}
-		if (!bGridTypeSupported)
+		else if (Grid->isType<openvdb::DoubleGrid>())
+		{
+			NumComponents = 1;
+			GridInfo.Format = EOpenVDBGridFormat::Double;
+		}
+		else if (Grid->isType<openvdb::Grid<openvdb::tree::Tree4<openvdb::Vec2d, 5, 4, 3>::Type>>())
 		{
+			NumComponents = 2;
+			GridInfo.Format = EOpenVDBGridFormat::Double;
+		}
+		else if (Grid->isType<openvdb::Vec3DGrid>())
+		{
+			NumComponents = 3;
+			GridInfo.Format = EOpenVDBGridFormat::Double;
+		}
+		else if (Grid->isType<openvdb::Grid<openvdb::tree::Tree4<openvdb::Vec4d, 5, 4, 3>::Type>>())
+		{
+			NumComponents = 4;
+			GridInfo.Format = EOpenVDBGridFormat::Double;
+		}
+#endif
+		else
+		{
+			// unsupported type
 			continue;
 		}
+		check(NumComponents > 0);
 
 		GridInfo.Name = Grid->getName().c_str();
 
-		FStringFormatOrderedArguments FormatArgs;
-		FormatArgs.Add(GridInfo.Name);
-		FormatArgs.Add(OpenVDBGridFormatToString(GridInfo.Format));
+		// Create one entry per component
+		for (uint32 ComponentIdx = 0; ComponentIdx < NumComponents; ++ComponentIdx)
+		{
+			GridInfo.ComponentIndex = ComponentIdx;
+
+			const TCHAR* ComponentNames[] = { TEXT(".x"), TEXT(".y"),TEXT(".z"),TEXT(".w") };
+			FStringFormatOrderedArguments FormatArgs;
+			FormatArgs.Add(GridInfo.Name);
+			FormatArgs.Add(NumComponents == 1 ? TEXT("") : ComponentNames[GridInfo.ComponentIndex]);
+			FormatArgs.Add(OpenVDBGridFormatToString(GridInfo.Format));
 
-		GridInfo.DisplayString = FString::Format(TEXT("Name: {0}, Format: {1}"), FormatArgs);
+			GridInfo.DisplayString = FString::Format(TEXT("{0}{1} ({2})"), FormatArgs);
 
-		OutGridInfo->Add(MoveTemp(GridInfo));
+			OutGridInfo->Add(MoveTemp(GridInfo));
+		}
 	}
 
 	if (OutGridInfo->IsEmpty())
@@ -192,9 +238,50 @@ bool GetOpenVDBGridInfo(TArray<uint8>& SourceFile, const FString& Filename, TArr
 	return false;
 }
 
+static EPixelFormat GetMultiComponentFormat(ESparseVolumePackedDataFormat Format, uint32 NumComponents)
+{
+	switch (Format)
+	{
+	case ESparseVolumePackedDataFormat::Unorm8:
+	{
+		switch (NumComponents)
+		{
+		case 1: return PF_R8;
+		case 2: return PF_R8G8;
+		case 3:
+		case 4: return PF_R8G8B8A8;
+		}
+		break;
+	}
+	case ESparseVolumePackedDataFormat::Float16:
+	{
+		switch (NumComponents)
+		{
+		case 1: return PF_R16F;
+		case 2: return PF_G16R16F;
+		case 3:
+		case 4: return PF_FloatRGBA;
+		}
+		break;
+	}
+	case ESparseVolumePackedDataFormat::Float32:
+	{
+		switch (NumComponents)
+		{
+		case 1: return PF_R32_FLOAT;
+		case 2: return PF_G32R32F;
+		case 3:
+		case 4: return PF_A32B32G32R32F;
+		}
+		break;
+	}
+	}
+	return PF_Unknown;
+}
+
 bool ConvertOpenVDBToSparseVolumeTexture(
 	TArray<uint8>& SourceFile,
-	uint32 GridIndex,
+	struct FSparseVolumeRawSourcePackedData& PackedData,
 	FSparseVolumeAssetHeader* OutHeader,
 	TArray<uint32>* OutDensityPage,
 	TArray<uint8>* OutDensityData,
@@ -203,59 +290,124 @@ bool ConvertOpenVDBToSparseVolumeTexture(
 	FVector ActiveMax)
 {
 #if OPENVDB_AVAILABLE
-	FArrayUint8StreamBuf StreamBuf(SourceFile);
-	std::istream IStream(&StreamBuf);
-	openvdb::io::Stream Stream(IStream, false /*delayLoad*/);
 
-	openvdb::GridPtrVecPtr Grids = Stream.getGrids();
-	if (!Grids || Grids->size() <= GridIndex)
+	uint32 NumRequiredComponents = 0;
+	for (uint32 ComponentIdx = 0; ComponentIdx < 4; ++ComponentIdx)
 	{
-		return false;
+		if (PackedData.SourceGridIndex[ComponentIdx] != INDEX_NONE)
+		{
+			NumRequiredComponents = FMath::Max(ComponentIdx + 1, NumRequiredComponents);
+		}
 	}
+	const uint32 NumActualComponents = NumRequiredComponents == 3 ? 4 : NumRequiredComponents; // We don't support formats with only 3 components
 
-	openvdb::FloatGrid::Ptr DensityGrid = openvdb::gridPtrCast<openvdb::FloatGrid>((*Grids)[GridIndex]);
-	if (!DensityGrid)
+	const bool bNormalizedFormat = PackedData.Format == ESparseVolumePackedDataFormat::Unorm8;
+	const EPixelFormat MultiCompFormat = GetMultiComponentFormat(PackedData.Format, NumActualComponents);
+	
+	if (MultiCompFormat == PF_Unknown)
 	{
+		// SVT_TODO error message about unsupported format
 		return false;
 	}
 
-	FOpenVDBData OVDBData = GetOpenVDBData(DensityGrid);
-	if (!IsOpenVDBDataValid(OVDBData, TEXT("")))
+	// Load file
+	FArrayUint8StreamBuf StreamBuf(SourceFile);
+	std::istream IStream(&StreamBuf);
+	openvdb::io::Stream Stream(IStream, false /*delayLoad*/);
+
+	// Check that at least one component has a valid grid assigned to it
+	openvdb::GridPtrVecPtr Grids = Stream.getGrids();
+	const size_t NumSourceGrids = Grids ? Grids->size() : 0;
+	if (NumSourceGrids <= PackedData.SourceGridIndex.X
+		&& NumSourceGrids <= PackedData.SourceGridIndex.Y
+		&& NumSourceGrids <= PackedData.SourceGridIndex.Z
+		&& NumSourceGrids <= PackedData.SourceGridIndex.W)
 	{
 		return false;
 	}
 
-	if (bOverrideActiveMinMax)
+	FSparseVolumeAssetHeader& Header = *OutHeader;
+	Header.PackedDataAFormat = MultiCompFormat;
+	Header.SourceVolumeResolution = FIntVector::ZeroValue;
+	
+	FIntVector SmallestAABBMin = FIntVector(INT32_MAX);
+
+	// Compute per source grid data of up to 4 different grids (one per component)
+	openvdb::FloatGrid::Ptr FloatGrids[4]{};
+	openvdb::FloatGrid::Ptr DummyGrid = nullptr;
+	float GridBackgroundValues[4]{};
+	float NormalizeFactors[4]{ 1.0f, 1.0f, 1.0f, 1.0f };
+	for (uint32 CompIdx = 0; CompIdx < 4; ++CompIdx)
 	{
-		OVDBData.VolumeActiveAABBMin = ActiveMin;
-		OVDBData.VolumeActiveAABBMax = ActiveMax;
-		OVDBData.VolumeActiveDim = ActiveMax - ActiveMin;
-	}
+		if (PackedData.SourceGridIndex[CompIdx] == INDEX_NONE)
+		{
+			FloatGrids[CompIdx] = nullptr;
+			continue;
+		}
 
-	FSparseVolumeAssetHeader& Header = *OutHeader;
+		openvdb::FloatGrid::Ptr FloatGrid = openvdb::gridPtrCast<openvdb::FloatGrid>((*Grids)[PackedData.SourceGridIndex[CompIdx]]);
+		if (!FloatGrid)
+		{
+			return false;
+		}
 
-	Header.SourceVolumeResolution = FIntVector(OVDBData.VolumeActiveDim);
-	// SVT_TODO handle error too high volume resolution assuming dx11 limitations?
+		FOpenVDBData OVDBData = GetOpenVDBData(FloatGrid);
+		if (!IsOpenVDBDataValid(OVDBData, TEXT("")))
+		{
+			return false;
+		}
 
+		if (bOverrideActiveMinMax)
+		{
+			OVDBData.VolumeActiveAABBMin = ActiveMin;
+			OVDBData.VolumeActiveAABBMax = ActiveMax;
+			OVDBData.VolumeActiveDim = ActiveMax - ActiveMin;
+		}
+
+		Header.SourceVolumeResolution.X = FMath::Max(Header.SourceVolumeResolution.X, OVDBData.VolumeActiveDim.X);
+		Header.SourceVolumeResolution.Y = FMath::Max(Header.SourceVolumeResolution.Y, OVDBData.VolumeActiveDim.Y);
+		Header.SourceVolumeResolution.Z = FMath::Max(Header.SourceVolumeResolution.Z, OVDBData.VolumeActiveDim.Z);
+		SmallestAABBMin.X = FMath::Min(SmallestAABBMin.X, OVDBData.VolumeActiveAABBMin.X);
+		SmallestAABBMin.Y = FMath::Min(SmallestAABBMin.Y, OVDBData.VolumeActiveAABBMin.Y);
+		SmallestAABBMin.Z = FMath::Min(SmallestAABBMin.Z, OVDBData.VolumeActiveAABBMin.Z);
+
+		FloatGrids[CompIdx] = FloatGrid;
+		DummyGrid = FloatGrid;
+		GridBackgroundValues[CompIdx] = FloatGrid->background();
+		if (bNormalizedFormat && PackedData.bRescaleInputForUnorm)
+		{
+			float MinVal = 0.0f;
+			float MaxVal = 0.0f;
+			FloatGrid->evalMinMax(MinVal, MaxVal);
+			NormalizeFactors[CompIdx] = MaxVal > SMALL_NUMBER ? (1.0f / MaxVal) : 1.0f;
+		}
+	}
+	
+	// SVT_TODO handle error too high volume resolution assuming dx11 limitations?
 	Header.PageTableVolumeResolution = FIntVector3(
 		FMath::DivideAndRoundUp(Header.SourceVolumeResolution.X, SPARSE_VOLUME_TILE_RES),
 		FMath::DivideAndRoundUp(Header.SourceVolumeResolution.Y, SPARSE_VOLUME_TILE_RES),
 		FMath::DivideAndRoundUp(Header.SourceVolumeResolution.Z, SPARSE_VOLUME_TILE_RES));
 	Header.TileDataVolumeResolution = FIntVector::ZeroValue;	// unknown for now
 
-	const uint32 VolumeTileDataBytes = SPARSE_VOLUME_TILE_RES * SPARSE_VOLUME_TILE_RES * SPARSE_VOLUME_TILE_RES;
+	const uint32 VolumeTileDataBytes = GPixelFormats[(SIZE_T)MultiCompFormat].Get3DImageSizeInBytes(SPARSE_VOLUME_TILE_RES, SPARSE_VOLUME_TILE_RES, SPARSE_VOLUME_TILE_RES);
+	const uint32 FormatSize = (uint32)GPixelFormats[(SIZE_T)MultiCompFormat].BlockBytes;
+	const uint32 SingleComponentFormatSize = (uint32)GPixelFormats[(SIZE_T)PackedData.Format].BlockBytes;
 	uint32 NumAllocatedPages = 0;
 
 	// Allocate some memory for temp data (worst case)
-	TArray<uint8> LinearAllocatedTiles;
 	TArray<FIntVector3> LinearAllocatedPages;
-	LinearAllocatedTiles.SetNum(Header.PageTableVolumeResolution.X * Header.PageTableVolumeResolution.Y * Header.PageTableVolumeResolution.Z * VolumeTileDataBytes);
 	LinearAllocatedPages.SetNum(Header.PageTableVolumeResolution.X * Header.PageTableVolumeResolution.Y * Header.PageTableVolumeResolution.Z);
 
-	auto accessor = DensityGrid->getConstAccessor();
-	const float GridBackgroundValue = DensityGrid->background();
+	openvdb::FloatGrid::ConstAccessor Accessors[4]
+	{
+		(FloatGrids[0] ? FloatGrids[0] : DummyGrid)->getConstAccessor(),
+		(FloatGrids[1] ? FloatGrids[1] : DummyGrid)->getConstAccessor(),
+		(FloatGrids[2] ? FloatGrids[2] : DummyGrid)->getConstAccessor(),
+		(FloatGrids[3] ? FloatGrids[3] : DummyGrid)->getConstAccessor(),
+	};
 
-	// Go over each potential page table form the source data and push allocate it if it has any data.
+	// Go over each potential page from the source data and push allocate it if it has any data.
 	// Otherwise point to the default empty page.
 	bool bAnyEmptyPageExists = false;
 	for (int32_t PageZ = 0; PageZ < Header.PageTableVolumeResolution.Z; ++PageZ)
@@ -266,24 +418,21 @@ bool ConvertOpenVDBToSparseVolumeTexture(
 			{
 				bool bHasAnyData = false;
 
-				uint8* NewTileData = LinearAllocatedTiles.GetData() + NumAllocatedPages * VolumeTileDataBytes;
-
 				for (int32_t z = 0; z < SPARSE_VOLUME_TILE_RES; ++z)
 				{
 					for (int32_t y = 0; y < SPARSE_VOLUME_TILE_RES; ++y)
 					{
 						for (int32_t x = 0; x < SPARSE_VOLUME_TILE_RES; ++x)
 						{
-							FVector VoxelCoord(PageX * SPARSE_VOLUME_TILE_RES + x, PageY * SPARSE_VOLUME_TILE_RES + y, PageZ * SPARSE_VOLUME_TILE_RES + z);	// This assumes sampling outside the boundary retuirn a default value
-							FVector VoxelCoord2 = OVDBData.VolumeActiveAABBMin + VoxelCoord;
-							float VoxelValue = accessor.getValue(openvdb::Coord(VoxelCoord2.X, VoxelCoord2.Y, VoxelCoord2.Z));
-
-							// clamp for now until we have fp16
-							VoxelValue = FMath::Clamp(VoxelValue, 0.0f, 1.0f);
-
-							bHasAnyData |= VoxelValue > GridBackgroundValue;
-
-							NewTileData[z * SPARSE_VOLUME_TILE_RES * SPARSE_VOLUME_TILE_RES + y * SPARSE_VOLUME_TILE_RES + x] = uint8(VoxelValue * 255.0f);
+							for (uint32 CompIdx = 0; CompIdx < NumActualComponents; ++CompIdx)
+							{
+								if (PackedData.SourceGridIndex[CompIdx] != INDEX_NONE)
+								{
+									FVector VoxelCoord = FVector(SmallestAABBMin) + FVector(PageX, PageY, PageZ) * SPARSE_VOLUME_TILE_RES + FVector(x, y, z);	// This assumes sampling outside the boundary returns a default value
+									float VoxelValue = Accessors[CompIdx].getValue(openvdb::Coord(VoxelCoord.X, VoxelCoord.Y, VoxelCoord.Z));
+									bHasAnyData |= VoxelValue > GridBackgroundValues[CompIdx];
+								}
+							}
 						}
 					}
 				}
@@ -315,7 +464,7 @@ bool ConvertOpenVDBToSparseVolumeTexture(
 
 	// Initialise the SparseVolumeTexture page and tile.
 	OutDensityPage->SetNumZeroed(Header.PageTableVolumeResolution.X * Header.PageTableVolumeResolution.Y * Header.PageTableVolumeResolution.Z);
-	OutDensityData->SetNumZeroed(Header.TileDataVolumeResolution.X * Header.TileDataVolumeResolution.Y * Header.TileDataVolumeResolution.Z);
+	OutDensityData->SetNumZeroed(Header.TileDataVolumeResolution.X * Header.TileDataVolumeResolution.Y * Header.TileDataVolumeResolution.Z * FormatSize);
 	uint32* DensityPagePtr = OutDensityPage->GetData();
 	uint8* DensityDataPtr = OutDensityData->GetData();
 
@@ -346,7 +495,6 @@ bool ConvertOpenVDBToSparseVolumeTexture(
 	for (uint32 i = 0; i < NumAllocatedPages; ++i)
 	{
 		FIntVector3 PageCoordToSplat = LinearAllocatedPages[i];
-		uint8* TileDataToSplat = LinearAllocatedTiles.GetData() + i * VolumeTileDataBytes;
 
 		// A page encodes the physical tile coord as unsigned int of 11 11 10 bits
 		// This means a page coord cannot be larger than 2047 for x and y and 1023 for z
@@ -361,19 +509,49 @@ bool ConvertOpenVDBToSparseVolumeTexture(
 				PageCoordToSplat.X
 			] = DestinationTileCoord32bit;
 
-		// Now copy the tile data
+		// Now copy the tile data from the source to the page atlas
 		for (int32_t z = 0; z < SPARSE_VOLUME_TILE_RES; ++z)
 		{
 			for (int32_t y = 0; y < SPARSE_VOLUME_TILE_RES; ++y)
 			{
 				for (int32_t x = 0; x < SPARSE_VOLUME_TILE_RES; ++x)
 				{
-					DensityDataPtr
-						[
-							(DestinationTileCoord.Z * SPARSE_VOLUME_TILE_RES + z) * Header.TileDataVolumeResolution.X * Header.TileDataVolumeResolution.Y +
+					for (uint32 CompIdx = 0; CompIdx < NumActualComponents; ++CompIdx)
+					{
+						float VoxelValue = 0.0f;
+						if (PackedData.SourceGridIndex[CompIdx] != INDEX_NONE)
+						{
+							FVector VoxelCoord = FVector(SmallestAABBMin) + FVector(PageCoordToSplat * SPARSE_VOLUME_TILE_RES) + FVector(x, y, z);	// This assumes sampling outside the boundary returns a default value
+							VoxelValue = Accessors[CompIdx].getValue(openvdb::Coord(VoxelCoord.X, VoxelCoord.Y, VoxelCoord.Z));
+						}
+
+						const SIZE_T DstCoord = ((DestinationTileCoord.Z * SPARSE_VOLUME_TILE_RES + z) * Header.TileDataVolumeResolution.X * Header.TileDataVolumeResolution.Y +
 							(DestinationTileCoord.Y * SPARSE_VOLUME_TILE_RES + y) * Header.TileDataVolumeResolution.X +
-							(DestinationTileCoord.X * SPARSE_VOLUME_TILE_RES + x)
-						] = TileDataToSplat[z * SPARSE_VOLUME_TILE_RES * SPARSE_VOLUME_TILE_RES + y * SPARSE_VOLUME_TILE_RES + x];
+							(DestinationTileCoord.X * SPARSE_VOLUME_TILE_RES + x)) * FormatSize + CompIdx * SingleComponentFormatSize;
+
+						switch (PackedData.Format)
+						{
+						case ESparseVolumePackedDataFormat::Unorm8:
+						{
+							DensityDataPtr[DstCoord] = uint8(FMath::Clamp(VoxelValue * NormalizeFactors[CompIdx], 0.0f, 1.0f) * 255.0f);
+							break;
+						}
+						case ESparseVolumePackedDataFormat::Float16:
+						{
+							const uint16 VoxelValue16FEncoded = FFloat16(VoxelValue).Encoded;
+							*((uint16*)(&DensityDataPtr[DstCoord])) = VoxelValue16FEncoded;
+							break;
+						}
+						case ESparseVolumePackedDataFormat::Float32:
+						{
+							*((float*)(&DensityDataPtr[DstCoord])) = VoxelValue;
+							break;
+						}
+						default:
+							checkNoEntry();
+							break;
+						}
+					}
 				}
 			}
 		}
@@ -394,6 +572,8 @@ const TCHAR* OpenVDBGridFormatToString(EOpenVDBGridFormat Format)
 	{
 	case EOpenVDBGridFormat::Float:
 		return TEXT("Float");
+	case EOpenVDBGridFormat::Double:
+		return TEXT("Double");
 	default:
 		return TEXT("Unknown");
 	}
diff --git a/Engine/Source/Editor/SparseVolumeTexture/Private/SparseVolumeTextureOpenVDBUtility.h b/Engine/Source/Editor/SparseVolumeTexture/Private/SparseVolumeTextureOpenVDBUtility.h
--- a/Engine/Source/Editor/SparseVolumeTexture/Private/SparseVolumeTextureOpenVDBUtility.h
+++ b/Engine/Source/Editor/SparseVolumeTexture/Private/SparseVolumeTextureOpenVDBUtility.h
@@ -19,14 +19,16 @@ struct FOpenVDBData
 enum class EOpenVDBGridFormat : uint8
 {
 	Float = 0,
+	Double = 1,
 };
 
 struct FOpenVDBGridInfo
 {
 	uint32 Index;
+	uint32 ComponentIndex;
 	EOpenVDBGridFormat Format;
 	FString Name;
-	FString DisplayString; // Contains both Name and Format
+	FString DisplayString; // Contains Name and Format and component in source grid
 };
 
 bool IsOpenVDBDataValid(FOpenVDBData& OpenVDBData, const FString& Filename);
@@ -37,7 +39,7 @@ bool GetOpenVDBGridInfo(TArray<uint8>& SourceFile, const FString& Filename, TArr
 
 bool ConvertOpenVDBToSparseVolumeTexture(
 	TArray<uint8>& SourceFile,
-	uint32 GridIndex,
+	struct FSparseVolumeRawSourcePackedData& PackedData,
 	struct FSparseVolumeAssetHeader* OutHeader,
 	TArray<uint32>* OutDensityPage,
 	TArray<uint8>* OutDensityData,
diff --git a/Engine/Source/Runtime/Engine/Classes/SparseVolumeTexture/SparseVolumeTexture.h b/Engine/Source/Runtime/Engine/Classes/SparseVolumeTexture/SparseVolumeTexture.h
--- a/Engine/Source/Runtime/Engine/Classes/SparseVolumeTexture/SparseVolumeTexture.h
+++ b/Engine/Source/Runtime/Engine/Classes/SparseVolumeTexture/SparseVolumeTexture.h
@@ -19,7 +19,7 @@ namespace UE { namespace Shader	{ enum class EValueType : uint8; } }
 
 DECLARE_DELEGATE_RetVal_EightParams(bool, FConvertOpenVDBToSparseVolumeTextureDelegate, 
 	TArray<uint8>& SourceFile,
-	uint32 GridIndex,
+	struct FSparseVolumeRawSourcePackedData& PackedData,
 	struct FSparseVolumeAssetHeader* OutHeader,
 	TArray<uint32>* OutDensityPage,
 	TArray<uint8>* OutDensityData,
@@ -29,11 +29,19 @@ DECLARE_DELEGATE_RetVal_EightParams(bool, FConvertOpenVDBToSparseVolumeTextureDe
 
 ENGINE_API FConvertOpenVDBToSparseVolumeTextureDelegate& OnConvertOpenVDBToSparseVolumeTexture();
 
+enum class ESparseVolumePackedDataFormat : uint8
+{
+	Unorm8 = 0,
+	Float16 = 1,
+	Float32 = 2,
+};
+
 struct ENGINE_API FSparseVolumeAssetHeader
 {
 	FIntVector3 PageTableVolumeResolution;
 	FIntVector3 TileDataVolumeResolution;
 	FIntVector3 SourceVolumeResolution;
+	EPixelFormat PackedDataAFormat;
 
 	// The current data format version for the header.
 	static const uint32 kVersion = 0;
@@ -47,16 +55,26 @@ struct ENGINE_API FSparseVolumeAssetHeader
 		: PageTableVolumeResolution(FIntVector3(0, 0, 0))
 		, TileDataVolumeResolution(FIntVector3(0, 0, 0))
 		, SourceVolumeResolution(FIntVector3(0, 0, 0))
+		, PackedDataAFormat(PF_Unknown)
 		, Version(kVersion)
 	{
 	}
 };
 
+// Describes how the grids of a source asset map to the components of a float4 of packed data in a SVT
+struct ENGINE_API FSparseVolumeRawSourcePackedData
+{
+	ESparseVolumePackedDataFormat Format;
+	FUintVector4 SourceGridIndex;
+	FUintVector4 SourceComponentIndex; // [0-3]
+	bool bRescaleInputForUnorm; // Just clamps to the [0-1] range if this is false
+};
+
 // The structure represent the source asset in high quality. It is used to cook the runtime data
 struct ENGINE_API FSparseVolumeRawSource
 {
-	uint32						DensityGridIndex;
-	TArray<uint8>				SourceAssetFile;
+	FSparseVolumeRawSourcePackedData PackedDataA;
+	TArray<uint8> SourceAssetFile;
 
 	// The current data format version for the raw source data.
 	static const uint32 kVersion = 0;
@@ -67,8 +85,7 @@ struct ENGINE_API FSparseVolumeRawSource
 	void Serialize(FArchive& Ar);
 
 	FSparseVolumeRawSource()
-		: DensityGridIndex()
-		, Version(kVersion)
+		: Version(kVersion)
 	{
 	}
 };
diff --git a/Engine/Source/Runtime/Engine/Private/SparseVolumeTexture/SparseVolumeTexture.cpp b/Engine/Source/Runtime/Engine/Private/SparseVolumeTexture/SparseVolumeTexture.cpp
--- a/Engine/Source/Runtime/Engine/Private/SparseVolumeTexture/SparseVolumeTexture.cpp
+++ b/Engine/Source/Runtime/Engine/Private/SparseVolumeTexture/SparseVolumeTexture.cpp
@@ -43,6 +43,13 @@ void FSparseVolumeAssetHeader::Serialize(FArchive& Ar)
 		Ar << PageTableVolumeResolution;
 		Ar << TileDataVolumeResolution;
 		Ar << SourceVolumeResolution;
+
+		uint32 FormatAsUint = (uint32)PackedDataAFormat;
+		Ar << FormatAsUint;
+		if (Ar.IsLoading())
+		{
+			PackedDataAFormat = static_cast<EPixelFormat>(FormatAsUint);
+		}
 	}
 	else
 	{
@@ -57,7 +64,16 @@ void FSparseVolumeRawSource::Serialize(FArchive& Ar)
 
 	if (Version == 0)
 	{
-		Ar << DensityGridIndex;
+		uint32 FormatAsUint = (uint32)PackedDataA.Format;
+		Ar << FormatAsUint;
+		if (Ar.IsLoading())
+		{
+			PackedDataA.Format = static_cast<ESparseVolumePackedDataFormat>(FormatAsUint);
+		}
+		Ar << PackedDataA.SourceGridIndex;
+		Ar << PackedDataA.SourceComponentIndex;
+		Ar << PackedDataA.bRescaleInputForUnorm;
+
 		Ar << SourceAssetFile;
 	}
 	else
@@ -233,7 +249,7 @@ void UStaticSparseVolumeTexture::ConvertRawSourceDataToSparseVolumeTextureRuntim
 		// Then, cook the runtime data and serialize it out.
 		const bool bSuccess = ConvertOpenVDBToSparseVolumeTextureDelegate.IsBound() && ConvertOpenVDBToSparseVolumeTextureDelegate.Execute(
 			SparseVolumeRawSource.SourceAssetFile,
-			SparseVolumeRawSource.DensityGridIndex,
+			SparseVolumeRawSource.PackedDataA,
 			&StaticFrame.SparseVolumeTextureRuntime.Header,
 			&StaticFrame.SparseVolumeTextureRuntime.DensityPage,
 			&StaticFrame.SparseVolumeTextureRuntime.DensityData,
@@ -402,7 +418,7 @@ void FSparseVolumeTextureSceneProxy::InitRHI()
 
 	// Tile data
 	{
-		EPixelFormat VoxelFormat = PF_R8;
+		EPixelFormat VoxelFormat = SparseVolumeTextureRuntime->Header.PackedDataAFormat;
 		FIntVector3 TileDataVolumeResolution = SparseVolumeTextureRuntime->Header.TileDataVolumeResolution;
 		const FRHITextureCreateDesc Desc =
 			FRHITextureCreateDesc::Create3D(TEXT("SparseVolumeTexture.TileData.RHITexture"),
@@ -536,7 +552,7 @@ void UAnimatedSparseVolumeTexture::ConvertRawSourceDataToSparseVolumeTextureRunt
 		// Then, cook the runtime data and serialize it out.
 		const bool bSuccess = ConvertOpenVDBToSparseVolumeTextureDelegate.IsBound() && ConvertOpenVDBToSparseVolumeTextureDelegate.Execute(
 			SparseVolumeRawSource.SourceAssetFile,
-			SparseVolumeRawSource.DensityGridIndex,
+			SparseVolumeRawSource.PackedDataA,
 			&Frame.SparseVolumeTextureRuntime.Header,
 			&Frame.SparseVolumeTextureRuntime.DensityPage,
 			&Frame.SparseVolumeTextureRuntime.DensityData,
diff --git a/Engine/Source/Runtime/Renderer/Private/SparseVolumeTexture/SparseVolumeTextureViewerComponent.cpp b/Engine/Source/Runtime/Renderer/Private/SparseVolumeTexture/SparseVolumeTextureViewerComponent.cpp
--- a/Engine/Source/Runtime/Renderer/Private/SparseVolumeTexture/SparseVolumeTextureViewerComponent.cpp
+++ b/Engine/Source/Runtime/Renderer/Private/SparseVolumeTexture/SparseVolumeTextureViewerComponent.cpp
@@ -29,6 +29,7 @@ USparseVolumeTextureViewerComponent::USparseVolumeTextureViewerComponent(const F
 	, SparseVolumeTexturePreview(nullptr)
 	, bAnimate(false)
 	, AnimationFrame(0.0f)
+	, ComponentToVisualize(0)
 	, SparseVolumeTextureViewerSceneProxy(nullptr)
 {
 	PrimaryComponentTick.bCanEverTick = true;
@@ -163,13 +164,16 @@ void USparseVolumeTextureViewerComponent::SendRenderTransformCommand()
 		FBoxSphereBounds ComponentBounds = CalcBounds(ToWorldMat);
 
 		FMatrix44f ToLocalMat = FMatrix44f(FTranslationMatrix(-ToWorldMat.GetTranslation()) * FRotationMatrix(FRotator(GetComponentToWorld().GetRotation().Inverse())) * FScaleMatrix(ComponentBounds.BoxExtent.Reciprocal()));
+		
+		uint32 CompToVisualize = (uint32)ComponentToVisualize;
 
 		FSparseVolumeTextureViewerSceneProxy* SVTViewerSceneProxy = SparseVolumeTextureViewerSceneProxy;
 		ENQUEUE_RENDER_COMMAND(FUpdateSparseVolumeTextureViewerProxyTransformCommand)(
-			[SVTViewerSceneProxy, ComponentBounds, ToLocalMat](FRHICommandList& RHICmdList)
+			[SVTViewerSceneProxy, ComponentBounds, ToLocalMat, CompToVisualize](FRHICommandList& RHICmdList)
 			{
 				SVTViewerSceneProxy->VolumeWorldBounds = ComponentBounds;
 				SVTViewerSceneProxy->WorldToLocal = ToLocalMat;
+				SVTViewerSceneProxy->ComponentToVisualize = CompToVisualize;
 			});
 	}
 }
diff --git a/Engine/Source/Runtime/Renderer/Private/SparseVolumeTexture/SparseVolumeTextureViewerComponent.h b/Engine/Source/Runtime/Renderer/Private/SparseVolumeTexture/SparseVolumeTextureViewerComponent.h
--- a/Engine/Source/Runtime/Renderer/Private/SparseVolumeTexture/SparseVolumeTextureViewerComponent.h
+++ b/Engine/Source/Runtime/Renderer/Private/SparseVolumeTexture/SparseVolumeTextureViewerComponent.h
@@ -35,6 +35,9 @@ class USparseVolumeTextureViewerComponent : public UPrimitiveComponent
 	UPROPERTY(EditAnywhere, BlueprintReadWrite, Category = "Asset Preview", meta = (UIMin = 0.0, UIMax = 1.0, ClampMin = 0.0, ClampMax = 1.0, EditCondition = "!bAnimate"))
 	float AnimationFrame;
 
+	UPROPERTY(EditAnywhere, BlueprintReadWrite, Category = "Asset Preview", meta = (UIMin = 0, UIMax = 4, ClampMin = 0, ClampMax = 3))
+	int32 ComponentToVisualize;
+
 	//~ Begin UPrimitiveComponent Interface.
 	virtual bool SupportsStaticLighting() const override { return false; }
 	//~ End UPrimitiveComponent Interface.
diff --git a/Engine/Source/Runtime/Renderer/Private/SparseVolumeTexture/SparseVolumeTextureViewerRendering.cpp b/Engine/Source/Runtime/Renderer/Private/SparseVolumeTexture/SparseVolumeTextureViewerRendering.cpp
--- a/Engine/Source/Runtime/Renderer/Private/SparseVolumeTexture/SparseVolumeTextureViewerRendering.cpp
+++ b/Engine/Source/Runtime/Renderer/Private/SparseVolumeTexture/SparseVolumeTextureViewerRendering.cpp
@@ -62,6 +62,7 @@ class FVisualizeSparseVolumeTexturePS : public FGlobalShader
 		SHADER_PARAMETER(FVector3f, VolumeBoundMaxWorld)
 		SHADER_PARAMETER(FVector3f, SparseVolumeTextureResolution)
 		SHADER_PARAMETER(FVector3f, SparseVolumeTexturePageTableResolution)
+		SHADER_PARAMETER(uint32, ComponentToVisualize)
 		SHADER_PARAMETER_TEXTURE(Texture3D, SparseVolumeTexture)
 		SHADER_PARAMETER_TEXTURE(Texture3D<uint>, SparseVolumeTexturePageTable)
 	END_SHADER_PARAMETER_STRUCT()
@@ -125,6 +126,7 @@ void AddSparseVolumeTextureViewerRenderPass(FRDGBuilder& GraphBuilder, FSceneRen
 			PsPassParameters->WorldToLocal0 = FVector4f(WorldToLocal.M[0][0], WorldToLocal.M[1][0], WorldToLocal.M[2][0], WorldToLocal.M[3][0]);
 			PsPassParameters->WorldToLocal1 = FVector4f(WorldToLocal.M[0][1], WorldToLocal.M[1][1], WorldToLocal.M[2][1], WorldToLocal.M[3][1]);
 			PsPassParameters->WorldToLocal2 = FVector4f(WorldToLocal.M[0][2], WorldToLocal.M[1][2], WorldToLocal.M[2][2], WorldToLocal.M[3][2]);
+			PsPassParameters->ComponentToVisualize = SVTProxy->ComponentToVisualize;
 			PsPassParameters->VolumeBoundMinWorld = FVector3f(SVTProxyBound.Origin - SVTProxyBound.BoxExtent);
 			PsPassParameters->VolumeBoundMaxWorld = FVector3f(SVTProxyBound.Origin + SVTProxyBound.BoxExtent);
 			PsPassParameters->SparseVolumeTextureResolution = FVector3f::OneVector;
diff --git a/Engine/Source/Runtime/Renderer/Private/SparseVolumeTexture/SparseVolumeTextureViewerSceneProxy.h b/Engine/Source/Runtime/Renderer/Private/SparseVolumeTexture/SparseVolumeTextureViewerSceneProxy.h
--- a/Engine/Source/Runtime/Renderer/Private/SparseVolumeTexture/SparseVolumeTextureViewerSceneProxy.h
+++ b/Engine/Source/Runtime/Renderer/Private/SparseVolumeTexture/SparseVolumeTextureViewerSceneProxy.h
@@ -20,6 +20,7 @@ public:
 	const FSparseVolumeTextureSceneProxy* SparseVolumeTextureSceneProxy = nullptr;
 	FBoxSphereBounds VolumeWorldBounds;
 	FMatrix44f WorldToLocal;
+	uint32 ComponentToVisualize;
 
 protected:
 
