Fixes for warnings in OpenVDB importer.

#rnx

diff --git a/Engine/Source/Editor/SparseVolumeTexture/Private/OpenVDBImportWindow.cpp b/Engine/Source/Editor/SparseVolumeTexture/Private/OpenVDBImportWindow.cpp
--- a/Engine/Source/Editor/SparseVolumeTexture/Private/OpenVDBImportWindow.cpp
+++ b/Engine/Source/Editor/SparseVolumeTexture/Private/OpenVDBImportWindow.cpp
@@ -331,7 +331,6 @@ void SOpenVDBComponentPicker::RefreshUIFromData()
 {
 	for (const TSharedPtr<FOpenVDBGridInfo>& Grid : *OpenVDBGridInfo)
 	{
-		printf("%i\n", Grid->Index);
 		if (Grid->Index == PackedData->SourceGridIndex[ComponentIndex] && Grid->ComponentIndex == PackedData->SourceComponentIndex[ComponentIndex])
 		{
 			GridComboBox->SetSelectedItem(Grid);
diff --git a/Engine/Source/Editor/SparseVolumeTexture/Private/SparseVolumeTextureOpenVDBUtility.cpp b/Engine/Source/Editor/SparseVolumeTexture/Private/SparseVolumeTextureOpenVDBUtility.cpp
--- a/Engine/Source/Editor/SparseVolumeTexture/Private/SparseVolumeTextureOpenVDBUtility.cpp
+++ b/Engine/Source/Editor/SparseVolumeTexture/Private/SparseVolumeTextureOpenVDBUtility.cpp
@@ -217,7 +217,7 @@ bool GetOpenVDBGridInfo(TArray<uint8>& SourceFile, const FString& Filename, TArr
 			const TCHAR* ComponentNames[] = { TEXT(".x"), TEXT(".y"),TEXT(".z"),TEXT(".w") };
 			FStringFormatOrderedArguments FormatArgs;
 			FormatArgs.Add(GridInfo.Name);
-			FormatArgs.Add(NumComponents == 1 ? TEXT("") : ComponentNames[GridInfo.ComponentIndex]);
+			FormatArgs.Add(TEXT("")); // SV_TODO: Silences warning about NumComponents always being 1. Replace with this later: FormatArgs.Add(NumComponents == 1 ? TEXT("") : ComponentNames[GridInfo.ComponentIndex]);
 			FormatArgs.Add(OpenVDBGridFormatToString(GridInfo.Format));
 
 			GridInfo.DisplayString = FString::Format(TEXT("{0}{1} ({2})"), FormatArgs);
