Implemented support for importing additional OpenVDB grid types (multi-component Float and Double). Also refactored some related functions and added a display of all the grids present in the source file to the import dialog.

#rnx

diff --git a/Engine/Source/Editor/SparseVolumeTexture/Private/OpenVDBGridAdapter.h b/Engine/Source/Editor/SparseVolumeTexture/Private/OpenVDBGridAdapter.h
new file mode 100644
--- /dev/null
+++ b/Engine/Source/Editor/SparseVolumeTexture/Private/OpenVDBGridAdapter.h
@@ -0,0 +1,179 @@
+// Copyright Epic Games, Inc. All Rights Reserved.
+
+#pragma once
+
+#if WITH_EDITOR
+
+#include "SparseVolumeTextureOpenVDB.h"
+
+#if OPENVDB_AVAILABLE
+
+#include "CoreMinimal.h"
+
+using FOpenVDBFloat1Grid = openvdb::FloatGrid;
+using FOpenVDBFloat2Grid = openvdb::Grid<openvdb::tree::Tree4<openvdb::Vec2f, 5, 4, 3>::Type>;
+using FOpenVDBFloat3Grid = openvdb::Vec3SGrid;
+using FOpenVDBFloat4Grid = openvdb::Grid<openvdb::tree::Tree4<openvdb::Vec4f, 5, 4, 3>::Type>;
+using FOpenVDBDouble1Grid = openvdb::DoubleGrid;
+using FOpenVDBDouble2Grid = openvdb::Grid<openvdb::tree::Tree4<openvdb::Vec2d, 5, 4, 3>::Type>;
+using FOpenVDBDouble3Grid = openvdb::Vec3DGrid;
+using FOpenVDBDouble4Grid = openvdb::Grid<openvdb::tree::Tree4<openvdb::Vec4d, 5, 4, 3>::Type>;
+
+template<typename ValueType>
+float GetOpenVDBValueComponent(const ValueType& Value, uint32 ComponentIndex)
+{
+	return (float)Value[FMath::Min(ComponentIndex, (uint32)ValueType::size)];
+}
+
+template<>
+float GetOpenVDBValueComponent<float>(const float& Value, uint32 ComponentIndex)
+{
+	return Value;
+}
+
+template<>
+float GetOpenVDBValueComponent<double>(const double& Value, uint32 ComponentIndex)
+{
+	return (float)Value;
+}
+
+template<typename ValueType>
+ValueType GetOpenVDBMaxComponent(const ValueType& ValueA, const ValueType& ValueB)
+{
+	return openvdb::math::maxComponent(ValueA, ValueB);
+}
+
+template<>
+float GetOpenVDBMaxComponent<float>(const float& ValueA, const float& ValueB)
+{
+	return FMath::Max(ValueA, ValueB);
+}
+
+template<>
+double GetOpenVDBMaxComponent<double>(const double& ValueA, const double& ValueB)
+{
+	return FMath::Max(ValueA, ValueB);
+}
+
+template<typename ValueType>
+ValueType GetOpenVDBMinComponent(const ValueType& ValueA, const ValueType& ValueB)
+{
+	return openvdb::math::minComponent(ValueA, ValueB);
+}
+
+template<>
+float GetOpenVDBMinComponent<float>(const float& ValueA, const float& ValueB)
+{
+	return FMath::Min(ValueA, ValueB);
+}
+
+template<>
+double GetOpenVDBMinComponent<double>(const double& ValueA, const double& ValueB)
+{
+	return FMath::Min(ValueA, ValueB);
+}
+
+// Interface for sampling and getting other values from an openVDB grid, abstracting away the actual type/format of the underlying grid.
+class IOpenVDBGridAdapterBase
+{
+public:
+	virtual float Sample(const openvdb::Coord& VoxelCoord, uint32 ComponentIndex) const = 0;
+	virtual void GetMinMaxValue(uint32 ComponentIndex, float* OutMinVal, float* OutMaxVal) const = 0;
+	virtual float GetBackgroundValue(uint32 ComponentIndex) const = 0;
+	virtual ~IOpenVDBGridAdapterBase() = default;
+};
+
+template<typename GridType>
+class TOpenVDBGridAdapter : public IOpenVDBGridAdapterBase
+{
+	using ValueType = typename GridType::ValueType;
+public:
+	explicit TOpenVDBGridAdapter(openvdb::SharedPtr<GridType> Grid)
+		: Grid(Grid), Accessor(Grid->getConstAccessor())
+	{
+	}
+
+	float Sample(const openvdb::Coord& VoxelCoord, uint32 ComponentIndex) const override
+	{
+		ValueType VoxelValue = Accessor.getValue(VoxelCoord);
+		return GetOpenVDBValueComponent(VoxelValue, ComponentIndex);
+	}
+
+	void GetMinMaxValue(uint32 ComponentIndex, float* OutMinVal, float* OutMaxVal) const override
+	{
+		*OutMinVal = FLT_MAX;
+		*OutMaxVal = -FLT_MAX;
+		if (auto Iter = Grid->tree().cbeginValueOn())
+		{
+			ValueType MinVal = *Iter;
+			ValueType MaxVal = MinVal;
+
+			for (++Iter; Iter; ++Iter)
+			{
+				ValueType Value = *Iter;
+				MinVal = GetOpenVDBMinComponent(MinVal, Value);
+				MaxVal = GetOpenVDBMaxComponent(MaxVal, Value);
+			}
+
+			*OutMinVal = GetOpenVDBValueComponent(MinVal, ComponentIndex);
+			*OutMaxVal = GetOpenVDBValueComponent(MaxVal, ComponentIndex);
+		}
+	}
+
+	float GetBackgroundValue(uint32 ComponentIndex) const override
+	{
+		return GetOpenVDBValueComponent(Grid->background(), ComponentIndex);
+	}
+
+private:
+	typename GridType::Ptr Grid;
+	typename GridType::ConstAccessor Accessor;
+};
+
+// This is just to not have to repeat the GridType multiple times per MakeShared<>() call.
+template<typename GridType>
+TSharedPtr<IOpenVDBGridAdapterBase> CreateOpenVDBGridAdapterInternal(openvdb::GridBase::Ptr Grid)
+{
+	return MakeShared<TOpenVDBGridAdapter<GridType>>(TOpenVDBGridAdapter<GridType>(openvdb::gridPtrCast<GridType>(Grid)));
+}
+
+// Creates an adapter suitable for the type of the given grid or nullptr if the type is not supported.
+TSharedPtr<IOpenVDBGridAdapterBase> CreateOpenVDBGridAdapter(openvdb::GridBase::Ptr Grid)
+{
+	if (Grid->isType<FOpenVDBFloat1Grid>())
+	{
+		return CreateOpenVDBGridAdapterInternal<FOpenVDBFloat1Grid>(Grid);
+	}
+	else if (Grid->isType<FOpenVDBFloat2Grid>())
+	{
+		return CreateOpenVDBGridAdapterInternal<FOpenVDBFloat2Grid>(Grid);
+	}
+	else if (Grid->isType<FOpenVDBFloat3Grid>())
+	{
+		return CreateOpenVDBGridAdapterInternal<FOpenVDBFloat3Grid>(Grid);
+	}
+	else if (Grid->isType<FOpenVDBFloat4Grid>())
+	{
+		return CreateOpenVDBGridAdapterInternal<FOpenVDBFloat4Grid>(Grid);
+	}
+	else if (Grid->isType<FOpenVDBDouble1Grid>())
+	{
+		return CreateOpenVDBGridAdapterInternal<FOpenVDBDouble1Grid>(Grid);
+	}
+	else if (Grid->isType<FOpenVDBDouble2Grid>())
+	{
+		return CreateOpenVDBGridAdapterInternal<FOpenVDBDouble2Grid>(Grid);
+	}
+	else if (Grid->isType<FOpenVDBDouble3Grid>())
+	{
+		return CreateOpenVDBGridAdapterInternal<FOpenVDBDouble3Grid>(Grid);
+	}
+	else if (Grid->isType<FOpenVDBDouble4Grid>())
+	{
+		return CreateOpenVDBGridAdapterInternal<FOpenVDBDouble4Grid>(Grid);
+	}
+	return nullptr;
+}
+
+#endif //OPENVDB_AVAILABLE
+#endif // WITH_EDITOR
diff --git a/Engine/Source/Editor/SparseVolumeTexture/Private/OpenVDBImportWindow.cpp b/Engine/Source/Editor/SparseVolumeTexture/Private/OpenVDBImportWindow.cpp
--- a/Engine/Source/Editor/SparseVolumeTexture/Private/OpenVDBImportWindow.cpp
+++ b/Engine/Source/Editor/SparseVolumeTexture/Private/OpenVDBImportWindow.cpp
@@ -17,7 +17,7 @@
 
 #define LOCTEXT_NAMESPACE "SOpenVDBImportWindow"
 
-static FText GetGridComboBoxItemText(TSharedPtr<FOpenVDBGridInfo> InItem)
+static FText GetGridComboBoxItemText(TSharedPtr<FOpenVDBGridComponentInfo> InItem)
 {
 	return InItem ? FText::FromString(InItem->DisplayString) : LOCTEXT("NoneGrid", "<None>");
 };
@@ -40,7 +40,7 @@ static FText GetFormatComboBoxItemText(TSharedPtr<ESparseVolumePackedDataFormat>
 void SOpenVDBImportWindow::Construct(const FArguments& InArgs)
 {
 	PackedDataA = InArgs._PackedDataA;
-	OpenVDBGridInfo = InArgs._OpenVDBGridInfo;
+	OpenVDBGridComponentInfo = InArgs._OpenVDBGridComponentInfo;
 	OpenVDBSupportedTargetFormats = InArgs._OpenVDBSupportedTargetFormats;
 	WidgetWindow = InArgs._WidgetWindow;
 
@@ -97,6 +97,25 @@ void SOpenVDBImportWindow::Construct(const FArguments& InArgs)
 			+ SVerticalBox::Slot()
 			.AutoHeight()
 			.Padding(2)
+			[
+				SNew(STextBlock)
+				.Text(LOCTEXT("OpenVDBImportWindow_FileInfo", "Source File Grid Info"))
+			]
+			+ SVerticalBox::Slot()
+			.AutoHeight()
+			.Padding(2)
+			[
+				SNew(SBorder)
+				.Padding(FMargin(3))
+				.BorderImage(FAppStyle::GetBrush("ToolPanel.GroupBorder"))
+				[
+					SNew(STextBlock)
+					.Text(FText::FromString(InArgs._FileInfoString))
+				]
+			]
+			+ SVerticalBox::Slot()
+			.AutoHeight()
+			.Padding(2)
 			[
 				SNew(SUniformGridPanel)
 				.SlotPadding(2)
@@ -127,7 +146,7 @@ void SOpenVDBImportWindow::Construct(const FArguments& InArgs)
 		[
 			SAssignNew(PackedDataAConfigurator, SOpenVDBPackedDataConfigurator)
 			.PackedData(PackedDataA)
-			.OpenVDBGridInfo(OpenVDBGridInfo)
+			.OpenVDBGridComponentInfo(OpenVDBGridComponentInfo)
 			.OpenVDBSupportedTargetFormats(OpenVDBSupportedTargetFormats)
 			.PackedDataName(LOCTEXT("OpenVDBImportWindow_PackedDataA", "Packed Data A"))
 		]
@@ -235,12 +254,12 @@ FText SOpenVDBImportWindow::GetImportTypeDisplayText() const
 
 void SOpenVDBImportWindow::SetDefaultGridAssignment()
 {
-	check(OpenVDBGridInfo);
+	check(OpenVDBGridComponentInfo);
 
 	PackedDataA->Format = ESparseVolumePackedDataFormat::Float32;
 	PackedDataA->SourceGridIndex = FUintVector4(INDEX_NONE);
 	PackedDataA->SourceComponentIndex = FUintVector4(INDEX_NONE);
-	PackedDataA->bRescaleInputForUnorm = false;
+	PackedDataA->bRemapInputForUnorm = false;
 
 	const TCHAR* SearchNames[] = { TEXT("density"), TEXT("heat"), TEXT("temperature"), TEXT("motion") };
 
@@ -249,7 +268,7 @@ void SOpenVDBImportWindow::SetDefaultGridAssignment()
 	{
 		for (uint32 InputComponent = 0; InputComponent < 4; ++InputComponent)
 		{
-			for (const TSharedPtr<FOpenVDBGridInfo>& Grid : *OpenVDBGridInfo)
+			for (const TSharedPtr<FOpenVDBGridComponentInfo>& Grid : *OpenVDBGridComponentInfo)
 			{
 				if (Grid->Name == SearchName && Grid->ComponentIndex == InputComponent)
 				{
@@ -269,7 +288,7 @@ void SOpenVDBComponentPicker::Construct(const FArguments& InArgs)
 {
 	PackedData = InArgs._PackedData;
 	ComponentIndex = InArgs._ComponentIndex;
-	OpenVDBGridInfo = InArgs._OpenVDBGridInfo;
+	OpenVDBGridComponentInfo = InArgs._OpenVDBGridComponentInfo;
 	
 	check(ComponentIndex < 4);
 	const TCHAR* ComponentLabels[] = { TEXT("X"), TEXT("Y"), TEXT("Z"), TEXT("W") };
@@ -295,14 +314,14 @@ void SOpenVDBComponentPicker::Construct(const FArguments& InArgs)
 				SNew(SBox)
 				.WidthOverride(300.0f)
 				[
-					SAssignNew(GridComboBox, SComboBox<TSharedPtr<FOpenVDBGridInfo>>)
-					.OptionsSource(OpenVDBGridInfo)
-					.OnGenerateWidget_Lambda([](TSharedPtr<FOpenVDBGridInfo> InItem)
+					SAssignNew(GridComboBox, SComboBox<TSharedPtr<FOpenVDBGridComponentInfo>>)
+					.OptionsSource(OpenVDBGridComponentInfo)
+					.OnGenerateWidget_Lambda([](TSharedPtr<FOpenVDBGridComponentInfo> InItem)
 					{
 						return SNew(STextBlock)
 						.Text(GetGridComboBoxItemText(InItem));
 					})
-					.OnSelectionChanged_Lambda([this](TSharedPtr<FOpenVDBGridInfo> InItem, ESelectInfo::Type)
+					.OnSelectionChanged_Lambda([this](TSharedPtr<FOpenVDBGridComponentInfo> InItem, ESelectInfo::Type)
 					{
 						if (InItem)
 						{
@@ -329,7 +348,7 @@ void SOpenVDBComponentPicker::Construct(const FArguments& InArgs)
 
 void SOpenVDBComponentPicker::RefreshUIFromData()
 {
-	for (const TSharedPtr<FOpenVDBGridInfo>& Grid : *OpenVDBGridInfo)
+	for (const TSharedPtr<FOpenVDBGridComponentInfo>& Grid : *OpenVDBGridComponentInfo)
 	{
 		if (Grid->Index == PackedData->SourceGridIndex[ComponentIndex] && Grid->ComponentIndex == PackedData->SourceComponentIndex[ComponentIndex])
 		{
@@ -350,7 +369,7 @@ void SOpenVDBPackedDataConfigurator::Construct(const FArguments& InArgs)
 			SNew(SOpenVDBComponentPicker)
 			.PackedData(PackedData)
 			.ComponentIndex(ComponentIndex)
-			.OpenVDBGridInfo(InArgs._OpenVDBGridInfo);
+			.OpenVDBGridComponentInfo(InArgs._OpenVDBGridComponentInfo);
 	}
 
 	this->ChildSlot
@@ -407,7 +426,11 @@ void SOpenVDBPackedDataConfigurator::Construct(const FArguments& InArgs)
 				.Padding(2.0f)
 				[
 					SNew(STextBlock)
-					.Text(LOCTEXT("UnormRescaleCheckBoxLabel", "Unorm Rescale"))
+					.Text(LOCTEXT("UnormRemapCheckBoxLabel", "Unorm Remap"))
+					.IsEnabled_Lambda([this]()
+					{
+						return PackedData->Format == ESparseVolumePackedDataFormat::Unorm8;
+					})
 				]
 				+ SHorizontalBox::Slot()
 				.VAlign(VAlign_Center)
@@ -415,16 +438,16 @@ void SOpenVDBPackedDataConfigurator::Construct(const FArguments& InArgs)
 				.AutoWidth()
 				.Padding(2.0f)
 				[
-					SAssignNew(RescaleUnormCheckBox, SCheckBox)
+					SAssignNew(RemapUnormCheckBox, SCheckBox)
 					.OnCheckStateChanged_Lambda([this](ECheckBoxState CheckBoxState)
 					{
-							PackedData->bRescaleInputForUnorm = CheckBoxState == ECheckBoxState::Checked;
+						PackedData->bRemapInputForUnorm = CheckBoxState == ECheckBoxState::Checked;
 					})
 					.IsEnabled_Lambda([this]()
 					{
 						return PackedData->Format == ESparseVolumePackedDataFormat::Unorm8;
 					})
-					.ToolTipText(LOCTEXT("UnormRescaleCheckBoxTooltip", "Rescales input values for unorm formats into the [0-1] range instead of clamping values outside this range."))
+					.ToolTipText(LOCTEXT("UnormRemapCheckBoxTooltip", "Remaps input values for unorm formats into the [0-1] range instead of clamping values outside this range."))
 					.IsChecked(false)
 				]
 			]
@@ -481,7 +504,7 @@ void SOpenVDBPackedDataConfigurator::RefreshUIFromData()
 	{
 		ComponentPickers[i]->RefreshUIFromData();
 	}
-	RescaleUnormCheckBox->SetIsChecked(PackedData->bRescaleInputForUnorm);
+	RemapUnormCheckBox->SetIsChecked(PackedData->bRemapInputForUnorm);
 }
 
 #undef LOCTEXT_NAMESPACE
diff --git a/Engine/Source/Editor/SparseVolumeTexture/Private/OpenVDBImportWindow.h b/Engine/Source/Editor/SparseVolumeTexture/Private/OpenVDBImportWindow.h
--- a/Engine/Source/Editor/SparseVolumeTexture/Private/OpenVDBImportWindow.h
+++ b/Engine/Source/Editor/SparseVolumeTexture/Private/OpenVDBImportWindow.h
@@ -7,22 +7,29 @@
 #include "Widgets/Input/SComboBox.h"
 #include "Widgets/Input/SCheckBox.h"
 
-struct FOpenVDBGridInfo;
 struct ENGINE_API FSparseVolumeRawSourcePackedData;
 enum class ESparseVolumePackedDataFormat : uint8;
 
+struct FOpenVDBGridComponentInfo
+{
+	uint32 Index;
+	uint32 ComponentIndex;
+	FString Name;
+	FString DisplayString; // Contains source file grid index, name and component (if it is a multi component type like Float3)
+};
+
 class SOpenVDBComponentPicker : public SCompoundWidget
 {
 public:
 	SLATE_BEGIN_ARGS(SOpenVDBComponentPicker)
 		: _PackedData()
 		, _ComponentIndex()
-		, _OpenVDBGridInfo()
+		, _OpenVDBGridComponentInfo()
 	{}
 
 		SLATE_ARGUMENT(FSparseVolumeRawSourcePackedData*, PackedData)
 		SLATE_ARGUMENT(uint32, ComponentIndex)
-		SLATE_ARGUMENT(TArray<TSharedPtr<FOpenVDBGridInfo>>*, OpenVDBGridInfo)
+		SLATE_ARGUMENT(TArray<TSharedPtr<FOpenVDBGridComponentInfo>>*, OpenVDBGridComponentInfo)
 	SLATE_END_ARGS()
 
 public:
@@ -32,8 +39,8 @@ public:
 private:
 	FSparseVolumeRawSourcePackedData* PackedData;
 	uint32 ComponentIndex;
-	TArray<TSharedPtr<FOpenVDBGridInfo>>* OpenVDBGridInfo;
-	TSharedPtr<SComboBox<TSharedPtr<FOpenVDBGridInfo>>>	GridComboBox;
+	TArray<TSharedPtr<FOpenVDBGridComponentInfo>>* OpenVDBGridComponentInfo;
+	TSharedPtr<SComboBox<TSharedPtr<FOpenVDBGridComponentInfo>>> GridComboBox;
 };
 
 class SOpenVDBPackedDataConfigurator : public SCompoundWidget
@@ -41,13 +48,13 @@ class SOpenVDBPackedDataConfigurator : public SCompoundWidget
 public:
 	SLATE_BEGIN_ARGS(SOpenVDBPackedDataConfigurator)
 		: _PackedData()
-		, _OpenVDBGridInfo()
+		, _OpenVDBGridComponentInfo()
 		, _OpenVDBSupportedTargetFormats()
 		, _PackedDataName()
 	{}
 
 		SLATE_ARGUMENT(FSparseVolumeRawSourcePackedData*, PackedData)
-		SLATE_ARGUMENT(TArray<TSharedPtr<FOpenVDBGridInfo>>*, OpenVDBGridInfo)
+		SLATE_ARGUMENT(TArray<TSharedPtr<FOpenVDBGridComponentInfo>>*, OpenVDBGridComponentInfo)
 		SLATE_ARGUMENT(TArray<TSharedPtr<ESparseVolumePackedDataFormat>>*, OpenVDBSupportedTargetFormats)
 		SLATE_ARGUMENT(FText, PackedDataName)
 	SLATE_END_ARGS()
@@ -61,7 +68,7 @@ private:
 	TSharedPtr<SOpenVDBComponentPicker> ComponentPickers[4];
 	TArray<TSharedPtr<ESparseVolumePackedDataFormat>>* OpenVDBSupportedTargetFormats;
 	TSharedPtr<SComboBox<TSharedPtr<ESparseVolumePackedDataFormat>>>	FormatComboBox;
-	TSharedPtr<SCheckBox> RescaleUnormCheckBox;
+	TSharedPtr<SCheckBox> RemapUnormCheckBox;
 };
 
 class SOpenVDBImportWindow : public SCompoundWidget
@@ -69,7 +76,8 @@ class SOpenVDBImportWindow : public SCompoundWidget
 public:
 	SLATE_BEGIN_ARGS(SOpenVDBImportWindow)
 		: _PackedDataA()
-		, _OpenVDBGridInfo()
+		, _OpenVDBGridComponentInfo()
+		, _FileInfoString()
 		, _OpenVDBSupportedTargetFormats()
 		, _WidgetWindow()
 		, _FullPath()
@@ -78,7 +86,8 @@ public:
 	{}
 
 		SLATE_ARGUMENT(FSparseVolumeRawSourcePackedData*, PackedDataA)
-		SLATE_ARGUMENT(TArray<TSharedPtr<FOpenVDBGridInfo>>*, OpenVDBGridInfo)
+		SLATE_ARGUMENT(TArray<TSharedPtr<FOpenVDBGridComponentInfo>>*, OpenVDBGridComponentInfo)
+		SLATE_ARGUMENT(FString, FileInfoString)
 		SLATE_ARGUMENT(TArray<TSharedPtr<ESparseVolumePackedDataFormat>>*, OpenVDBSupportedTargetFormats)
 		SLATE_ARGUMENT(TSharedPtr<SWindow>, WidgetWindow)
 		SLATE_ARGUMENT(FText, FullPath)
@@ -95,7 +104,7 @@ public:
 
 private:
 	FSparseVolumeRawSourcePackedData*					PackedDataA;
-	TArray<TSharedPtr<FOpenVDBGridInfo>>*				OpenVDBGridInfo;
+	TArray<TSharedPtr<FOpenVDBGridComponentInfo>>*		OpenVDBGridComponentInfo;
 	TArray<TSharedPtr<ESparseVolumePackedDataFormat>>*	OpenVDBSupportedTargetFormats;
 	TSharedPtr<SOpenVDBPackedDataConfigurator>			PackedDataAConfigurator;
 	TSharedPtr<SButton>									ImportButton;
diff --git a/Engine/Source/Editor/SparseVolumeTexture/Private/SparseVolumeTextureFactory.cpp b/Engine/Source/Editor/SparseVolumeTexture/Private/SparseVolumeTextureFactory.cpp
--- a/Engine/Source/Editor/SparseVolumeTexture/Private/SparseVolumeTextureFactory.cpp
+++ b/Engine/Source/Editor/SparseVolumeTexture/Private/SparseVolumeTextureFactory.cpp
@@ -223,43 +223,87 @@ UObject* USparseVolumeTextureFactory::FactoryCreateFile(UClass* InClass, UObject
 
 		// Load file and get info about each contained grid
 		TArray<uint8> LoadedFile;
-		TArray<TSharedPtr<FOpenVDBGridInfo>> GridInfoPtrs;
+		TArray<TSharedPtr<FOpenVDBGridComponentInfo>> GridComponentInfoPtrs;
 		TArray<TSharedPtr<ESparseVolumePackedDataFormat>> SupportedFormats =
 		{ 
 			MakeShared<ESparseVolumePackedDataFormat>(ESparseVolumePackedDataFormat::Float32),
 			MakeShared<ESparseVolumePackedDataFormat>(ESparseVolumePackedDataFormat::Float16),
 			MakeShared<ESparseVolumePackedDataFormat>(ESparseVolumePackedDataFormat::Unorm8)
 		};
+		FString FileInfoString;
 		{
 			if (!FFileHelper::LoadFileToArray(LoadedFile, *Filename))
 			{
-				UE_LOG(LogSparseVolumeTextureFactory, Error, TEXT("OpenVDB file could not be opened: %s"), *Filename);
+				UE_LOG(LogSparseVolumeTextureFactory, Error, TEXT("OpenVDB file could not be loaded: %s"), *Filename);
 				return nullptr;
 			}
 
 			TArray<FOpenVDBGridInfo> GridInfo;
-			if (!GetOpenVDBGridInfo(LoadedFile, Filename, &GridInfo))
+			if (!GetOpenVDBGridInfo(LoadedFile, &GridInfo))
 			{
-				UE_LOG(LogSparseVolumeTextureFactory, Error, TEXT("OpenVDB file contains no suitable density grids: %s"), *Filename);
+				UE_LOG(LogSparseVolumeTextureFactory, Error, TEXT("Failed to read OpenVDB file: %s"), *Filename);
+				return nullptr;
+			}
+			if (GridInfo.IsEmpty())
+			{
+				UE_LOG(LogSparseVolumeTextureFactory, Error, TEXT("OpenVDB file contains no grids: %s"), *Filename);
+				return nullptr;
+			}
+
+			// Create individual entries for each component of all valid source grids
+			TArray<FOpenVDBGridComponentInfo> GridComponentInfo;
+			for (const FOpenVDBGridInfo& Grid : GridInfo)
+			{
+				// Append all grids to the string, even if we don't actually support them
+				FileInfoString.Append(Grid.DisplayString);
+				FileInfoString.AppendChar(TEXT('\n'));
+
+				if (Grid.Type == EOpenVDBGridType::Unknown || !IsOpenVDBDataValid(Grid.OpenVDBData, Filename))
+				{
+					continue;
+				}
+
+				// Create one entry per component
+				for (uint32 ComponentIdx = 0; ComponentIdx < Grid.NumComponents; ++ComponentIdx)
+				{
+					FOpenVDBGridComponentInfo ComponentInfo;
+					ComponentInfo.Index = Grid.Index;
+					ComponentInfo.ComponentIndex = ComponentIdx;
+					ComponentInfo.Name = Grid.Name;
+
+					const TCHAR* ComponentNames[] = { TEXT(".X"), TEXT(".Y"),TEXT(".Z"),TEXT(".W") };
+					FStringFormatOrderedArguments FormatArgs;
+					FormatArgs.Add(ComponentInfo.Index);
+					FormatArgs.Add(ComponentInfo.Name);
+					FormatArgs.Add(Grid.NumComponents == 1 ? TEXT("") : ComponentNames[ComponentIdx]);
+
+					ComponentInfo.DisplayString = FString::Format(TEXT("{0}. {1}{2}"), FormatArgs);
+
+					GridComponentInfo.Add(MoveTemp(ComponentInfo));
+				}
+			}
+
+			if (GridComponentInfo.IsEmpty())
+			{
+				UE_LOG(LogSparseVolumeTextureFactory, Error, TEXT("OpenVDB file contains no grids of supported type: %s"), *Filename);
 				return nullptr;
 			}
 
 			// Convert to array of TSharedPtr. SComboBox requires this.
-			GridInfoPtrs.Empty(GridInfo.Num() + 1);
+			GridComponentInfoPtrs.Empty(GridComponentInfo.Num() + 1);
 			
 			// We need a <None> option to leave channels empty
-			FOpenVDBGridInfo NoneGridInfo;
-			NoneGridInfo.Index = INDEX_NONE;
-			NoneGridInfo.ComponentIndex = INDEX_NONE;
-			NoneGridInfo.Format = EOpenVDBGridFormat::Float;
-			NoneGridInfo.Name = TEXT("<None>");
-			NoneGridInfo.DisplayString = TEXT("<None>");
+			FOpenVDBGridComponentInfo NoneGridComponentInfo;
+			NoneGridComponentInfo.Index = INDEX_NONE;
+			NoneGridComponentInfo.ComponentIndex = INDEX_NONE;
+			NoneGridComponentInfo.Name = TEXT("<None>");
+			NoneGridComponentInfo.DisplayString = TEXT("<None>");
 			
-			GridInfoPtrs.Add(MakeShared<FOpenVDBGridInfo>(NoneGridInfo));
+			GridComponentInfoPtrs.Add(MakeShared<FOpenVDBGridComponentInfo>(NoneGridComponentInfo));
 			
-			for (auto& Grid : GridInfo)
+			for (const FOpenVDBGridComponentInfo& GridCompInfo : GridComponentInfo)
 			{
-				GridInfoPtrs.Add(MakeShared<FOpenVDBGridInfo>(Grid));
+				GridComponentInfoPtrs.Add(MakeShared<FOpenVDBGridComponentInfo>(GridCompInfo));
 			}
 		}
 
@@ -301,7 +345,8 @@ UObject* USparseVolumeTextureFactory::FactoryCreateFile(UClass* InClass, UObject
 			(
 				SAssignNew(OpenVDBOptionWindow, SOpenVDBImportWindow)
 				.PackedDataA(&PackedDataA)
-				.OpenVDBGridInfo(&GridInfoPtrs)
+				.OpenVDBGridComponentInfo(&GridComponentInfoPtrs)
+				.FileInfoString(FileInfoString)
 				.OpenVDBSupportedTargetFormats(&SupportedFormats)
 				.WidgetWindow(Window)
 				.FullPath(FText::FromString(Filename))
diff --git a/Engine/Source/Editor/SparseVolumeTexture/Private/SparseVolumeTextureOpenVDBUtility.cpp b/Engine/Source/Editor/SparseVolumeTexture/Private/SparseVolumeTextureOpenVDBUtility.cpp
--- a/Engine/Source/Editor/SparseVolumeTexture/Private/SparseVolumeTextureOpenVDBUtility.cpp
+++ b/Engine/Source/Editor/SparseVolumeTexture/Private/SparseVolumeTextureOpenVDBUtility.cpp
@@ -5,6 +5,7 @@
 #include "SparseVolumeTextureOpenVDBUtility.h"
 #include "SparseVolumeTextureOpenVDB.h"
 #include "SparseVolumeTexture/SparseVolumeTexture.h"
+#include "OpenVDBGridAdapter.h"
 
 DEFINE_LOG_CATEGORY_STATIC(LogSparseVolumeTextureOpenVDBUtility, Log, All);
 
@@ -43,18 +44,18 @@ static FOpenVDBData GetOpenVDBData(openvdb::GridBase::Ptr GridBase)
 
 #endif
 
-bool IsOpenVDBDataValid(FOpenVDBData& OpenVDBData, const FString& Filename)
+bool IsOpenVDBDataValid(const FOpenVDBData& OpenVDBData, const FString& Filename)
 {
 	if (OpenVDBData.VolumeActiveDim.X * OpenVDBData.VolumeActiveDim.Y * OpenVDBData.VolumeActiveDim.Z == 0)
 	{
 		// SVT_TODO we should gently handle that case
-		UE_LOG(LogSparseVolumeTextureOpenVDBUtility, Error, TEXT("OpenVDB asset is empty due to volume size being 0: %s"), *Filename);
+		UE_LOG(LogSparseVolumeTextureOpenVDBUtility, Warning, TEXT("OpenVDB grid is empty due to volume size being 0: %s"), *Filename);
 		return false;
 	}
 
 	if (!OpenVDBData.bHasUniformVoxels)
 	{
-		UE_LOG(LogSparseVolumeTextureOpenVDBUtility, Error, TEXT("OpenVDB importer cannot handle non uniform voxels: %s"), *Filename);
+		UE_LOG(LogSparseVolumeTextureOpenVDBUtility, Warning, TEXT("OpenVDB importer cannot handle non uniform voxels: %s"), *Filename);
 		return false;
 	}
 	return true;
@@ -133,7 +134,7 @@ bool FindDensityGridIndex(TArray<uint8>& SourceFile, const FString& Filename, ui
 	return false;
 }
 
-bool GetOpenVDBGridInfo(TArray<uint8>& SourceFile, const FString& Filename, TArray<FOpenVDBGridInfo>* OutGridInfo)
+bool GetOpenVDBGridInfo(TArray<uint8>& SourceFile, TArray<FOpenVDBGridInfo>* OutGridInfo)
 {
 #if OPENVDB_AVAILABLE
 	FArrayUint8StreamBuf StreamBuf(SourceFile);
@@ -142,94 +143,69 @@ bool GetOpenVDBGridInfo(TArray<uint8>& SourceFile, const FString& Filename, TArr
 
 	openvdb::GridPtrVecPtr Grids = Stream.getGrids();
 
-	if (!Grids)
-	{
-		UE_LOG(LogSparseVolumeTextureOpenVDBUtility, Warning, TEXT("OpenVDB file contains no grids: %s"), *Filename);
-		return false;
-	}
-
-	OutGridInfo->Empty(Grids->size() * 2 /*assume an average of 2 channels per grid*/);
+	OutGridInfo->Empty(Grids->size());
 
 	uint32 GridIndex = 0;
 	for (openvdb::GridBase::Ptr& Grid : *Grids)
 	{
+		// Create entry
 		FOpenVDBGridInfo GridInfo;
 		GridInfo.Index = GridIndex++;
+		GridInfo.NumComponents = 0;
+		GridInfo.Type = EOpenVDBGridType::Unknown;
+		GridInfo.Name = Grid->getName().c_str();
+		GridInfo.OpenVDBData = GetOpenVDBData(Grid);
 
 		// Figure out the type/format of the grid
-		uint32 NumComponents = 0;
-		if (Grid->isType<openvdb::FloatGrid>())
-		{
-			NumComponents = 1;
-			GridInfo.Format = EOpenVDBGridFormat::Float;
-		}
-#if 0 // SV_TODO: actually support these additional types in ConvertOpenVDBToSparseVolumeTexture()
-		else if (Grid->isType<openvdb::Grid<openvdb::tree::Tree4<openvdb::Vec2f, 5, 4, 3>::Type>>())
+		if (Grid->isType<FOpenVDBFloat1Grid>())
 		{
-			NumComponents = 2;
-			GridInfo.Format = EOpenVDBGridFormat::Float;
+			GridInfo.NumComponents = 1;
+			GridInfo.Type = EOpenVDBGridType::Float;
 		}
-		else if (Grid->isType<openvdb::Vec3SGrid>())
+		else if (Grid->isType<FOpenVDBFloat2Grid>())
 		{
-			NumComponents = 3;
-			GridInfo.Format = EOpenVDBGridFormat::Float;
+			GridInfo.NumComponents = 2;
+			GridInfo.Type = EOpenVDBGridType::Float2;
 		}
-		else if (Grid->isType<openvdb::Grid<openvdb::tree::Tree4<openvdb::Vec4f, 5, 4, 3>::Type>>())
+		else if (Grid->isType<FOpenVDBFloat3Grid>())
 		{
-			NumComponents = 4;
-			GridInfo.Format = EOpenVDBGridFormat::Float;
+			GridInfo.NumComponents = 3;
+			GridInfo.Type = EOpenVDBGridType::Float3;
 		}
-		else if (Grid->isType<openvdb::DoubleGrid>())
+		else if (Grid->isType<FOpenVDBFloat4Grid>())
 		{
-			NumComponents = 1;
-			GridInfo.Format = EOpenVDBGridFormat::Double;
+			GridInfo.NumComponents = 4;
+			GridInfo.Type = EOpenVDBGridType::Float4;
 		}
-		else if (Grid->isType<openvdb::Grid<openvdb::tree::Tree4<openvdb::Vec2d, 5, 4, 3>::Type>>())
+		else if (Grid->isType<FOpenVDBDouble1Grid>())
 		{
-			NumComponents = 2;
-			GridInfo.Format = EOpenVDBGridFormat::Double;
+			GridInfo.NumComponents = 1;
+			GridInfo.Type = EOpenVDBGridType::Double;
 		}
-		else if (Grid->isType<openvdb::Vec3DGrid>())
+		else if (Grid->isType<FOpenVDBDouble2Grid>())
 		{
-			NumComponents = 3;
-			GridInfo.Format = EOpenVDBGridFormat::Double;
+			GridInfo.NumComponents = 2;
+			GridInfo.Type = EOpenVDBGridType::Double2;
 		}
-		else if (Grid->isType<openvdb::Grid<openvdb::tree::Tree4<openvdb::Vec4d, 5, 4, 3>::Type>>())
+		else if (Grid->isType<FOpenVDBDouble3Grid>())
 		{
-			NumComponents = 4;
-			GridInfo.Format = EOpenVDBGridFormat::Double;
+			GridInfo.NumComponents = 3;
+			GridInfo.Type = EOpenVDBGridType::Double3;
 		}
-#endif
-		else
+		else if (Grid->isType<FOpenVDBDouble4Grid>())
 		{
-			// unsupported type
-			continue;
+			GridInfo.NumComponents = 4;
+			GridInfo.Type = EOpenVDBGridType::Double4;
 		}
-		check(NumComponents > 0);
-
-		GridInfo.Name = Grid->getName().c_str();
-
-		// Create one entry per component
-		for (uint32 ComponentIdx = 0; ComponentIdx < NumComponents; ++ComponentIdx)
-		{
-			GridInfo.ComponentIndex = ComponentIdx;
 
-			const TCHAR* ComponentNames[] = { TEXT(".x"), TEXT(".y"),TEXT(".z"),TEXT(".w") };
-			FStringFormatOrderedArguments FormatArgs;
-			FormatArgs.Add(GridInfo.Name);
-			FormatArgs.Add(TEXT("")); // SV_TODO: Silences warning about NumComponents always being 1. Replace with this later: FormatArgs.Add(NumComponents == 1 ? TEXT("") : ComponentNames[GridInfo.ComponentIndex]);
-			FormatArgs.Add(OpenVDBGridFormatToString(GridInfo.Format));
+		FStringFormatOrderedArguments FormatArgs;
+		FormatArgs.Add(GridInfo.Index);
+		FormatArgs.Add(OpenVDBGridTypeToString(GridInfo.Type));
+		FormatArgs.Add(GridInfo.Name);
 
-			GridInfo.DisplayString = FString::Format(TEXT("{0}{1} ({2})"), FormatArgs);
-
-			OutGridInfo->Add(MoveTemp(GridInfo));
-		}
-	}
+		GridInfo.DisplayString = FString::Format(TEXT("{0}. Type: {1}, Name: \"{2}\""), FormatArgs);
 
-	if (OutGridInfo->IsEmpty())
-	{
-		UE_LOG(LogSparseVolumeTextureOpenVDBUtility, Warning, TEXT("OpenVDB file contains no grids of supported type/format: %s"), *Filename);
-		return false;
+		OutGridInfo->Add(MoveTemp(GridInfo));
 	}
 
 	return true;
@@ -333,25 +309,26 @@ bool ConvertOpenVDBToSparseVolumeTexture(
 	FIntVector SmallestAABBMin = FIntVector(INT32_MAX);
 
 	// Compute per source grid data of up to 4 different grids (one per component)
-	openvdb::FloatGrid::Ptr FloatGrids[4]{};
-	openvdb::FloatGrid::Ptr DummyGrid = nullptr;
+	TSharedPtr<IOpenVDBGridAdapterBase> GridAdapters[4]{};
 	float GridBackgroundValues[4]{};
-	float NormalizeFactors[4]{ 1.0f, 1.0f, 1.0f, 1.0f };
+	float NormalizeScale[4]{ 1.0f, 1.0f, 1.0f, 1.0f };
+	float NormalizeBias[4]{};
 	for (uint32 CompIdx = 0; CompIdx < 4; ++CompIdx)
 	{
 		if (PackedData.SourceGridIndex[CompIdx] == INDEX_NONE)
 		{
-			FloatGrids[CompIdx] = nullptr;
+			GridAdapters[CompIdx] = nullptr;
 			continue;
 		}
 
-		openvdb::FloatGrid::Ptr FloatGrid = openvdb::gridPtrCast<openvdb::FloatGrid>((*Grids)[PackedData.SourceGridIndex[CompIdx]]);
-		if (!FloatGrid)
+		openvdb::GridBase::Ptr GridBase = (*Grids)[PackedData.SourceGridIndex[CompIdx]];
+		GridAdapters[CompIdx] = CreateOpenVDBGridAdapter(GridBase);
+		if (!GridAdapters[CompIdx])
 		{
 			return false;
 		}
 
-		FOpenVDBData OVDBData = GetOpenVDBData(FloatGrid);
+		FOpenVDBData OVDBData = GetOpenVDBData(GridBase);
 		if (!IsOpenVDBDataValid(OVDBData, TEXT("")))
 		{
 			return false;
@@ -371,15 +348,15 @@ bool ConvertOpenVDBToSparseVolumeTexture(
 		SmallestAABBMin.Y = FMath::Min(SmallestAABBMin.Y, OVDBData.VolumeActiveAABBMin.Y);
 		SmallestAABBMin.Z = FMath::Min(SmallestAABBMin.Z, OVDBData.VolumeActiveAABBMin.Z);
 
-		FloatGrids[CompIdx] = FloatGrid;
-		DummyGrid = FloatGrid;
-		GridBackgroundValues[CompIdx] = FloatGrid->background();
-		if (bNormalizedFormat && PackedData.bRescaleInputForUnorm)
+		GridBackgroundValues[CompIdx] = GridAdapters[CompIdx]->GetBackgroundValue(PackedData.SourceComponentIndex[CompIdx]);
+		if (bNormalizedFormat && PackedData.bRemapInputForUnorm)
 		{
 			float MinVal = 0.0f;
 			float MaxVal = 0.0f;
-			FloatGrid->evalMinMax(MinVal, MaxVal);
-			NormalizeFactors[CompIdx] = MaxVal > SMALL_NUMBER ? (1.0f / MaxVal) : 1.0f;
+			GridAdapters[CompIdx]->GetMinMaxValue(PackedData.SourceComponentIndex[CompIdx], &MinVal, &MaxVal);
+			const float Diff = MaxVal - MinVal;
+			NormalizeScale[CompIdx] = MaxVal > SMALL_NUMBER ? (1.0f / Diff) : 1.0f;
+			NormalizeBias[CompIdx] = -MinVal * NormalizeScale[CompIdx];
 		}
 	}
 	
@@ -392,21 +369,13 @@ bool ConvertOpenVDBToSparseVolumeTexture(
 
 	const uint32 VolumeTileDataBytes = GPixelFormats[(SIZE_T)MultiCompFormat].Get3DImageSizeInBytes(SPARSE_VOLUME_TILE_RES, SPARSE_VOLUME_TILE_RES, SPARSE_VOLUME_TILE_RES);
 	const uint32 FormatSize = (uint32)GPixelFormats[(SIZE_T)MultiCompFormat].BlockBytes;
-	const uint32 SingleComponentFormatSize = (uint32)GPixelFormats[(SIZE_T)PackedData.Format].BlockBytes;
+	const uint32 SingleComponentFormatSize = FormatSize / NumActualComponents;
 	uint32 NumAllocatedPages = 0;
 
 	// Allocate some memory for temp data (worst case)
 	TArray<FIntVector3> LinearAllocatedPages;
 	LinearAllocatedPages.SetNum(Header.PageTableVolumeResolution.X * Header.PageTableVolumeResolution.Y * Header.PageTableVolumeResolution.Z);
 
-	openvdb::FloatGrid::ConstAccessor Accessors[4]
-	{
-		(FloatGrids[0] ? FloatGrids[0] : DummyGrid)->getConstAccessor(),
-		(FloatGrids[1] ? FloatGrids[1] : DummyGrid)->getConstAccessor(),
-		(FloatGrids[2] ? FloatGrids[2] : DummyGrid)->getConstAccessor(),
-		(FloatGrids[3] ? FloatGrids[3] : DummyGrid)->getConstAccessor(),
-	};
-
 	// Go over each potential page from the source data and push allocate it if it has any data.
 	// Otherwise point to the default empty page.
 	bool bAnyEmptyPageExists = false;
@@ -429,8 +398,8 @@ bool ConvertOpenVDBToSparseVolumeTexture(
 								if (PackedData.SourceGridIndex[CompIdx] != INDEX_NONE)
 								{
 									FVector VoxelCoord = FVector(SmallestAABBMin) + FVector(PageX, PageY, PageZ) * SPARSE_VOLUME_TILE_RES + FVector(x, y, z);	// This assumes sampling outside the boundary returns a default value
-									float VoxelValue = Accessors[CompIdx].getValue(openvdb::Coord(VoxelCoord.X, VoxelCoord.Y, VoxelCoord.Z));
-									bHasAnyData |= VoxelValue > GridBackgroundValues[CompIdx];
+									float VoxelValue = GridAdapters[CompIdx]->Sample(openvdb::Coord(VoxelCoord.X, VoxelCoord.Y, VoxelCoord.Z), PackedData.SourceComponentIndex[CompIdx]);
+									bHasAnyData |= VoxelValue != GridBackgroundValues[CompIdx];
 								}
 							}
 						}
@@ -519,10 +488,12 @@ bool ConvertOpenVDBToSparseVolumeTexture(
 					for (uint32 CompIdx = 0; CompIdx < NumActualComponents; ++CompIdx)
 					{
 						float VoxelValue = 0.0f;
+						float VoxelValueNormalized = 0.0f;
 						if (PackedData.SourceGridIndex[CompIdx] != INDEX_NONE)
 						{
 							FVector VoxelCoord = FVector(SmallestAABBMin) + FVector(PageCoordToSplat * SPARSE_VOLUME_TILE_RES) + FVector(x, y, z);	// This assumes sampling outside the boundary returns a default value
-							VoxelValue = Accessors[CompIdx].getValue(openvdb::Coord(VoxelCoord.X, VoxelCoord.Y, VoxelCoord.Z));
+							VoxelValue = GridAdapters[CompIdx]->Sample(openvdb::Coord(VoxelCoord.X, VoxelCoord.Y, VoxelCoord.Z), PackedData.SourceComponentIndex[CompIdx]);
+							VoxelValueNormalized = FMath::Clamp(VoxelValue * NormalizeScale[CompIdx] + NormalizeBias[CompIdx], 0.0f, 1.0f);
 						}
 
 						const SIZE_T DstCoord = ((DestinationTileCoord.Z * SPARSE_VOLUME_TILE_RES + z) * Header.TileDataVolumeResolution.X * Header.TileDataVolumeResolution.Y +
@@ -533,7 +504,7 @@ bool ConvertOpenVDBToSparseVolumeTexture(
 						{
 						case ESparseVolumePackedDataFormat::Unorm8:
 						{
-							DensityDataPtr[DstCoord] = uint8(FMath::Clamp(VoxelValue * NormalizeFactors[CompIdx], 0.0f, 1.0f) * 255.0f);
+							DensityDataPtr[DstCoord] = uint8(VoxelValueNormalized * 255.0f);
 							break;
 						}
 						case ESparseVolumePackedDataFormat::Float16:
@@ -566,14 +537,26 @@ bool ConvertOpenVDBToSparseVolumeTexture(
 #endif // OPENVDB_AVAILABLE
 }
 
-const TCHAR* OpenVDBGridFormatToString(EOpenVDBGridFormat Format)
+const TCHAR* OpenVDBGridTypeToString(EOpenVDBGridType Type)
 {
-	switch (Format)
+	switch (Type)
 	{
-	case EOpenVDBGridFormat::Float:
+	case EOpenVDBGridType::Float:
 		return TEXT("Float");
-	case EOpenVDBGridFormat::Double:
+	case EOpenVDBGridType::Float2:
+		return TEXT("Float2");
+	case EOpenVDBGridType::Float3:
+		return TEXT("Float3");
+	case EOpenVDBGridType::Float4:
+		return TEXT("Float4");
+	case EOpenVDBGridType::Double:
 		return TEXT("Double");
+	case EOpenVDBGridType::Double2:
+		return TEXT("Double2");
+	case EOpenVDBGridType::Double3:
+		return TEXT("Double3");
+	case EOpenVDBGridType::Double4:
+		return TEXT("Double4");
 	default:
 		return TEXT("Unknown");
 	}
diff --git a/Engine/Source/Editor/SparseVolumeTexture/Private/SparseVolumeTextureOpenVDBUtility.h b/Engine/Source/Editor/SparseVolumeTexture/Private/SparseVolumeTextureOpenVDBUtility.h
--- a/Engine/Source/Editor/SparseVolumeTexture/Private/SparseVolumeTextureOpenVDBUtility.h
+++ b/Engine/Source/Editor/SparseVolumeTexture/Private/SparseVolumeTextureOpenVDBUtility.h
@@ -16,26 +16,34 @@ struct FOpenVDBData
 	bool bHasUniformVoxels;
 };
 
-enum class EOpenVDBGridFormat : uint8
+enum class EOpenVDBGridType : uint8
 {
-	Float = 0,
-	Double = 1,
+	Unknown = 0,
+	Float = 1,
+	Float2 = 2,
+	Float3 = 3,
+	Float4 = 4,
+	Double = 5,
+	Double2 = 6,
+	Double3 = 7,
+	Double4 = 8,
 };
 
 struct FOpenVDBGridInfo
 {
 	uint32 Index;
-	uint32 ComponentIndex;
-	EOpenVDBGridFormat Format;
+	uint32 NumComponents;
+	EOpenVDBGridType Type;
 	FString Name;
-	FString DisplayString; // Contains Name and Format and component in source grid
+	FString DisplayString; // Contains Index (into source file grids), Type and Name
+	FOpenVDBData OpenVDBData;
 };
 
-bool IsOpenVDBDataValid(FOpenVDBData& OpenVDBData, const FString& Filename);
+bool IsOpenVDBDataValid(const FOpenVDBData& OpenVDBData, const FString& Filename);
 
 bool FindDensityGridIndex(TArray<uint8>& SourceFile, const FString& Filename, uint32* OutGridIndex, FOpenVDBData* OutOVDBData);
 
-bool GetOpenVDBGridInfo(TArray<uint8>& SourceFile, const FString& Filename, TArray<FOpenVDBGridInfo>* OutGridInfo);
+bool GetOpenVDBGridInfo(TArray<uint8>& SourceFile, TArray<FOpenVDBGridInfo>* OutGridInfo);
 
 bool ConvertOpenVDBToSparseVolumeTexture(
 	TArray<uint8>& SourceFile,
@@ -47,6 +55,6 @@ bool ConvertOpenVDBToSparseVolumeTexture(
 	FVector ActiveMin,
 	FVector ActiveMax);
 
-const TCHAR* OpenVDBGridFormatToString(EOpenVDBGridFormat Format);
+const TCHAR* OpenVDBGridTypeToString(EOpenVDBGridType Type);
 
 #endif // WITH_EDITOR
\ No newline at end of file
diff --git a/Engine/Source/Runtime/Engine/Classes/SparseVolumeTexture/SparseVolumeTexture.h b/Engine/Source/Runtime/Engine/Classes/SparseVolumeTexture/SparseVolumeTexture.h
--- a/Engine/Source/Runtime/Engine/Classes/SparseVolumeTexture/SparseVolumeTexture.h
+++ b/Engine/Source/Runtime/Engine/Classes/SparseVolumeTexture/SparseVolumeTexture.h
@@ -66,8 +66,8 @@ struct ENGINE_API FSparseVolumeRawSourcePackedData
 {
 	ESparseVolumePackedDataFormat Format;
 	FUintVector4 SourceGridIndex;
-	FUintVector4 SourceComponentIndex; // [0-3]
-	bool bRescaleInputForUnorm; // Just clamps to the [0-1] range if this is false
+	FUintVector4 SourceComponentIndex; // Index [0-3] of the component in the source grid.
+	bool bRemapInputForUnorm; // Maps the input from its minimum and maximum value into the [0-1] range. Clamps to [0-1] otherwise.
 };
 
 // The structure represent the source asset in high quality. It is used to cook the runtime data
diff --git a/Engine/Source/Runtime/Engine/Private/SparseVolumeTexture/SparseVolumeTexture.cpp b/Engine/Source/Runtime/Engine/Private/SparseVolumeTexture/SparseVolumeTexture.cpp
--- a/Engine/Source/Runtime/Engine/Private/SparseVolumeTexture/SparseVolumeTexture.cpp
+++ b/Engine/Source/Runtime/Engine/Private/SparseVolumeTexture/SparseVolumeTexture.cpp
@@ -72,7 +72,7 @@ void FSparseVolumeRawSource::Serialize(FArchive& Ar)
 		}
 		Ar << PackedDataA.SourceGridIndex;
 		Ar << PackedDataA.SourceComponentIndex;
-		Ar << PackedDataA.bRescaleInputForUnorm;
+		Ar << PackedDataA.bRemapInputForUnorm;
 
 		Ar << SourceAssetFile;
 	}
