Added second float4/PackedDataB to SparseVolumeTextures.

#rnx

diff --git a/Engine/Shaders/Private/SparseVolumeTexture/VisualizeSparseVolumeTexture.usf b/Engine/Shaders/Private/SparseVolumeTexture/VisualizeSparseVolumeTexture.usf
--- a/Engine/Shaders/Private/SparseVolumeTexture/VisualizeSparseVolumeTexture.usf
+++ b/Engine/Shaders/Private/SparseVolumeTexture/VisualizeSparseVolumeTexture.usf
@@ -48,7 +48,8 @@ float4 WorldToLocal2;
 uint ComponentToVisualize;
 
 
-Texture3D<float4> SparseVolumeTexture;
+Texture3D<float4> SparseVolumeTextureA;
+Texture3D<float4> SparseVolumeTextureB;
 float3 SparseVolumeTextureResolution;
 
 Texture3D<uint> SparseVolumeTexturePageTable;
@@ -179,9 +180,18 @@ void VisualizeSparseVolumeTexturePS(
 
 				int3 VoxelCoord = TileCoord * SPARSE_VOLUME_TILE_RES + int3(P - PageCoordF * SPARSE_VOLUME_TILE_RES);
 
-				const float4 VoxelData = SparseVolumeTexture.Load(int4(VoxelCoord, 0));
+				float Density = 0.0f;
+				if (ComponentToVisualize < 4)
+				{
+					float4 VoxelData = SparseVolumeTextureA.Load(int4(VoxelCoord, 0));
+					Density = VoxelData[ComponentToVisualize];
+				}
+				else
+				{
+					float4 VoxelData = SparseVolumeTextureB.Load(int4(VoxelCoord, 0));
+					Density = VoxelData[min(ComponentToVisualize - 4, 3)];
+				}
 
-				const float Density = VoxelData[ComponentToVisualize];
 				if (Density > 0.0f)
 
 				{
diff --git a/Engine/Source/Editor/SparseVolumeTexture/Private/OpenVDBImportWindow.cpp b/Engine/Source/Editor/SparseVolumeTexture/Private/OpenVDBImportWindow.cpp
--- a/Engine/Source/Editor/SparseVolumeTexture/Private/OpenVDBImportWindow.cpp
+++ b/Engine/Source/Editor/SparseVolumeTexture/Private/OpenVDBImportWindow.cpp
@@ -40,6 +40,7 @@ static FText GetFormatComboBoxItemText(TSharedPtr<ESparseVolumePackedDataFormat>
 void SOpenVDBImportWindow::Construct(const FArguments& InArgs)
 {
 	PackedDataA = InArgs._PackedDataA;
+	PackedDataB = InArgs._PackedDataB;
 	OpenVDBGridComponentInfo = InArgs._OpenVDBGridComponentInfo;
 	OpenVDBSupportedTargetFormats = InArgs._OpenVDBSupportedTargetFormats;
 	WidgetWindow = InArgs._WidgetWindow;
@@ -150,6 +151,16 @@ void SOpenVDBImportWindow::Construct(const FArguments& InArgs)
 			.OpenVDBSupportedTargetFormats(OpenVDBSupportedTargetFormats)
 			.PackedDataName(LOCTEXT("OpenVDBImportWindow_PackedDataA", "Packed Data A"))
 		]
+		+ SVerticalBox::Slot()
+		.AutoHeight()
+		.Padding(2)
+		[
+			SAssignNew(PackedDataBConfigurator, SOpenVDBPackedDataConfigurator)
+			.PackedData(PackedDataB)
+			.OpenVDBGridComponentInfo(OpenVDBGridComponentInfo)
+			.OpenVDBSupportedTargetFormats(OpenVDBSupportedTargetFormats)
+			.PackedDataName(LOCTEXT("OpenVDBImportWindow_PackedDataB", "Packed Data B"))
+		]
 	);
 
 	SetDefaultGridAssignment();
@@ -229,14 +240,19 @@ EActiveTimerReturnType SOpenVDBImportWindow::SetFocusPostConstruct(double InCurr
 
 bool SOpenVDBImportWindow::CanImport() const
 {
-	const FUintVector4& GridIndices = PackedDataA->SourceGridIndex;
-	const FUintVector4& ComponentIndices = PackedDataA->SourceComponentIndex;
-	if ((GridIndices.X != INDEX_NONE && ComponentIndices.X != INDEX_NONE)
-		|| (GridIndices.Y != INDEX_NONE && ComponentIndices.Y != INDEX_NONE)
-		|| (GridIndices.Z != INDEX_NONE && ComponentIndices.Z != INDEX_NONE)
-		|| (GridIndices.W != INDEX_NONE && ComponentIndices.W != INDEX_NONE))
+	const FUintVector4& GridIndicesA = PackedDataA->SourceGridIndex;
+	const FUintVector4& ComponentIndicesA = PackedDataA->SourceComponentIndex;
+	const FUintVector4& GridIndicesB = PackedDataB->SourceGridIndex;
+	const FUintVector4& ComponentIndicesB = PackedDataB->SourceComponentIndex;
+
+	for (uint32 i = 0; i < 4; ++i)
 	{
-		return true;
+		const bool bGridAValid = GridIndicesA[i] != INDEX_NONE && ComponentIndicesA[i] != INDEX_NONE;
+		const bool bGridBValid = GridIndicesB[i] != INDEX_NONE && ComponentIndicesB[i] != INDEX_NONE;
+		if (bGridAValid || bGridBValid)
+		{
+			return true;
+		}
 	}
 	return false;
 }
@@ -256,25 +272,36 @@ void SOpenVDBImportWindow::SetDefaultGridAssignment()
 {
 	check(OpenVDBGridComponentInfo);
 
-	PackedDataA->Format = ESparseVolumePackedDataFormat::Float32;
-	PackedDataA->SourceGridIndex = FUintVector4(INDEX_NONE);
-	PackedDataA->SourceComponentIndex = FUintVector4(INDEX_NONE);
-	PackedDataA->bRemapInputForUnorm = false;
+	FSparseVolumeRawSourcePackedData* PackedData[] = { PackedDataA , PackedDataB };
+	for (FSparseVolumeRawSourcePackedData* Data : PackedData)
+	{
+		Data->Format = ESparseVolumePackedDataFormat::Float32;
+		Data->SourceGridIndex = FUintVector4(INDEX_NONE);
+		Data->SourceComponentIndex = FUintVector4(INDEX_NONE);
+		Data->bRemapInputForUnorm = false;
+	}
 
 	const TCHAR* SearchNames[] = { TEXT("density"), TEXT("heat"), TEXT("temperature"), TEXT("motion") };
 
+	// Try to find grids named like the SearchNames in the source file and assign them to the components of the output SVT.
+	uint32 CurrentOutputPackedData = 0;
 	uint32 CurrentOutputComponent = 0;
 	for (const TCHAR* SearchName : SearchNames)
 	{
-		for (uint32 InputComponent = 0; InputComponent < 4; ++InputComponent)
+		for (const TSharedPtr<FOpenVDBGridComponentInfo>& GridComponent : *OpenVDBGridComponentInfo)
 		{
-			for (const TSharedPtr<FOpenVDBGridComponentInfo>& Grid : *OpenVDBGridComponentInfo)
+			for (uint32 InputComponent = 0; InputComponent < 4; ++InputComponent)
 			{
-				if (Grid->Name == SearchName && Grid->ComponentIndex == InputComponent)
+				if (GridComponent->Name == SearchName && GridComponent->ComponentIndex == InputComponent)
 				{
-					PackedDataA->SourceGridIndex[CurrentOutputComponent] = Grid->Index;
-					PackedDataA->SourceComponentIndex[CurrentOutputComponent] = Grid->ComponentIndex;
+					PackedData[CurrentOutputPackedData]->SourceGridIndex[CurrentOutputComponent] = GridComponent->Index;
+					PackedData[CurrentOutputPackedData]->SourceComponentIndex[CurrentOutputComponent] = GridComponent->ComponentIndex;
 					++CurrentOutputComponent;
+					if (CurrentOutputComponent == 4)
+					{
+						CurrentOutputComponent = 0;
+						++CurrentOutputPackedData;
+					}
 					break;
 				}
 			}
@@ -282,6 +309,7 @@ void SOpenVDBImportWindow::SetDefaultGridAssignment()
 	}
 
 	PackedDataAConfigurator->RefreshUIFromData();
+	PackedDataBConfigurator->RefreshUIFromData();
 }
 
 void SOpenVDBComponentPicker::Construct(const FArguments& InArgs)
diff --git a/Engine/Source/Editor/SparseVolumeTexture/Private/OpenVDBImportWindow.h b/Engine/Source/Editor/SparseVolumeTexture/Private/OpenVDBImportWindow.h
--- a/Engine/Source/Editor/SparseVolumeTexture/Private/OpenVDBImportWindow.h
+++ b/Engine/Source/Editor/SparseVolumeTexture/Private/OpenVDBImportWindow.h
@@ -76,6 +76,7 @@ class SOpenVDBImportWindow : public SCompoundWidget
 public:
 	SLATE_BEGIN_ARGS(SOpenVDBImportWindow)
 		: _PackedDataA()
+		, _PackedDataB()
 		, _OpenVDBGridComponentInfo()
 		, _FileInfoString()
 		, _OpenVDBSupportedTargetFormats()
@@ -86,6 +87,7 @@ public:
 	{}
 
 		SLATE_ARGUMENT(FSparseVolumeRawSourcePackedData*, PackedDataA)
+		SLATE_ARGUMENT(FSparseVolumeRawSourcePackedData*, PackedDataB)
 		SLATE_ARGUMENT(TArray<TSharedPtr<FOpenVDBGridComponentInfo>>*, OpenVDBGridComponentInfo)
 		SLATE_ARGUMENT(FString, FileInfoString)
 		SLATE_ARGUMENT(TArray<TSharedPtr<ESparseVolumePackedDataFormat>>*, OpenVDBSupportedTargetFormats)
@@ -104,9 +106,11 @@ public:
 
 private:
 	FSparseVolumeRawSourcePackedData*					PackedDataA;
+	FSparseVolumeRawSourcePackedData*					PackedDataB;
 	TArray<TSharedPtr<FOpenVDBGridComponentInfo>>*		OpenVDBGridComponentInfo;
 	TArray<TSharedPtr<ESparseVolumePackedDataFormat>>*	OpenVDBSupportedTargetFormats;
 	TSharedPtr<SOpenVDBPackedDataConfigurator>			PackedDataAConfigurator;
+	TSharedPtr<SOpenVDBPackedDataConfigurator>			PackedDataBConfigurator;
 	TSharedPtr<SButton>									ImportButton;
 	TWeakPtr<SWindow>									WidgetWindow;
 	bool												bShouldImport;
diff --git a/Engine/Source/Editor/SparseVolumeTexture/Private/SparseVolumeTextureFactory.cpp b/Engine/Source/Editor/SparseVolumeTexture/Private/SparseVolumeTextureFactory.cpp
--- a/Engine/Source/Editor/SparseVolumeTexture/Private/SparseVolumeTextureFactory.cpp
+++ b/Engine/Source/Editor/SparseVolumeTexture/Private/SparseVolumeTextureFactory.cpp
@@ -196,6 +196,9 @@ UObject* USparseVolumeTextureFactory::FactoryCreateFile(UClass* InClass, UObject
 			SparseVolumeRawSource.PackedDataA.Format = ESparseVolumePackedDataFormat::Unorm8;
 			SparseVolumeRawSource.PackedDataA.SourceGridIndex = FUintVector4(0, INDEX_NONE, INDEX_NONE, INDEX_NONE);
 			SparseVolumeRawSource.PackedDataA.SourceComponentIndex = FUintVector4(0, INDEX_NONE, INDEX_NONE, INDEX_NONE);
+			SparseVolumeRawSource.PackedDataB.Format = ESparseVolumePackedDataFormat::Unorm8;
+			SparseVolumeRawSource.PackedDataB.SourceGridIndex = FUintVector4(INDEX_NONE, INDEX_NONE, INDEX_NONE, INDEX_NONE);
+			SparseVolumeRawSource.PackedDataB.SourceComponentIndex = FUintVector4(INDEX_NONE, INDEX_NONE, INDEX_NONE, INDEX_NONE);
 			SparseVolumeRawSource.SourceAssetFile = MoveTemp(LoadedFile);
 			LoadedFile.Reset();
 
@@ -309,6 +312,7 @@ UObject* USparseVolumeTextureFactory::FactoryCreateFile(UClass* InClass, UObject
 
 		// Show import dialog
 		FSparseVolumeRawSourcePackedData PackedDataA{};
+		FSparseVolumeRawSourcePackedData PackedDataB{};
 		{
 			TSharedPtr<SWindow> ParentWindow;
 
@@ -345,6 +349,7 @@ UObject* USparseVolumeTextureFactory::FactoryCreateFile(UClass* InClass, UObject
 			(
 				SAssignNew(OpenVDBOptionWindow, SOpenVDBImportWindow)
 				.PackedDataA(&PackedDataA)
+				.PackedDataB(&PackedDataB)
 				.OpenVDBGridComponentInfo(&GridComponentInfoPtrs)
 				.FileInfoString(FileInfoString)
 				.OpenVDBSupportedTargetFormats(&SupportedFormats)
@@ -368,6 +373,7 @@ UObject* USparseVolumeTextureFactory::FactoryCreateFile(UClass* InClass, UObject
 
 		FSparseVolumeRawSource SparseVolumeRawSource{};
 		SparseVolumeRawSource.PackedDataA = PackedDataA;
+		SparseVolumeRawSource.PackedDataB = PackedDataB;
 		SparseVolumeRawSource.SourceAssetFile = MoveTemp(LoadedFile);
 		LoadedFile.Reset();
 
diff --git a/Engine/Source/Editor/SparseVolumeTexture/Private/SparseVolumeTextureOpenVDBUtility.cpp b/Engine/Source/Editor/SparseVolumeTexture/Private/SparseVolumeTextureOpenVDBUtility.cpp
--- a/Engine/Source/Editor/SparseVolumeTexture/Private/SparseVolumeTextureOpenVDBUtility.cpp
+++ b/Engine/Source/Editor/SparseVolumeTexture/Private/SparseVolumeTextureOpenVDBUtility.cpp
@@ -257,32 +257,62 @@ static EPixelFormat GetMultiComponentFormat(ESparseVolumePackedDataFormat Format
 
 bool ConvertOpenVDBToSparseVolumeTexture(
 	TArray<uint8>& SourceFile,
-	struct FSparseVolumeRawSourcePackedData& PackedData,
-	FSparseVolumeAssetHeader* OutHeader,
-	TArray<uint32>* OutDensityPage,
-	TArray<uint8>* OutDensityData,
+	FSparseVolumeRawSourcePackedData& PackedDataA,
+	FSparseVolumeRawSourcePackedData& PackedDataB,
+	FOpenVDBToSVTConversionResult* OutResult,
 	bool bOverrideActiveMinMax,
 	FVector ActiveMin,
 	FVector ActiveMax)
 {
 #if OPENVDB_AVAILABLE
 
-	uint32 NumRequiredComponents = 0;
-	for (uint32 ComponentIdx = 0; ComponentIdx < 4; ++ComponentIdx)
+	constexpr uint32 NumPackedData = 2; // PackedDataA and PackedDataB, representing the two textures with voxel data
+	FSparseVolumeRawSourcePackedData* PackedData[NumPackedData] = { &PackedDataA, &PackedDataB };
+
+	// Compute some basic info about the number of components and which format to use
+	uint32 NumActualComponents[NumPackedData] = {};
+	EPixelFormat MultiCompFormat[NumPackedData] = {};
+	uint32 FormatSize[NumPackedData] = {};
+	uint32 SingleComponentFormatSize[NumPackedData] = {};
+	bool bNormalizedFormat[NumPackedData] = {};
+	bool bHasValidSourceGrids[NumPackedData] = {};
+	bool bAnySourceGridIndicesValid = false;
+	
+	for (uint32 PackedDataIdx = 0; PackedDataIdx < NumPackedData; ++PackedDataIdx)
 	{
-		if (PackedData.SourceGridIndex[ComponentIdx] != INDEX_NONE)
+		uint32 NumRequiredComponents = 0;
+		for (uint32 ComponentIdx = 0; ComponentIdx < 4; ++ComponentIdx)
+		{
+			if (PackedData[PackedDataIdx]->SourceGridIndex[ComponentIdx] != INDEX_NONE)
+			{
+				check(PackedData[PackedDataIdx]->SourceComponentIndex[ComponentIdx] != INDEX_NONE);
+				NumRequiredComponents = FMath::Max(ComponentIdx + 1, NumRequiredComponents);
+				bHasValidSourceGrids[PackedDataIdx] = true;
+				bAnySourceGridIndicesValid = true;
+			}
+		}
+
+		if (bHasValidSourceGrids[PackedDataIdx])
 		{
-			NumRequiredComponents = FMath::Max(ComponentIdx + 1, NumRequiredComponents);
+			NumActualComponents[PackedDataIdx] = NumRequiredComponents == 3 ? 4 : NumRequiredComponents; // We don't support formats with only 3 components
+			bNormalizedFormat[PackedDataIdx] = PackedData[PackedDataIdx]->Format == ESparseVolumePackedDataFormat::Unorm8;
+			MultiCompFormat[PackedDataIdx] = GetMultiComponentFormat(PackedData[PackedDataIdx]->Format, NumActualComponents[PackedDataIdx]);
+
+			if (MultiCompFormat[PackedDataIdx] == PF_Unknown)
+			{
+				// SVT_TODO error message about unsupported format
+				return false;
+			}
+
+			FormatSize[PackedDataIdx] = (uint32)GPixelFormats[(SIZE_T)MultiCompFormat[PackedDataIdx]].BlockBytes;
+			SingleComponentFormatSize[PackedDataIdx] = FormatSize[PackedDataIdx] / NumActualComponents[PackedDataIdx];
 		}
 	}
-	const uint32 NumActualComponents = NumRequiredComponents == 3 ? 4 : NumRequiredComponents; // We don't support formats with only 3 components
 
-	const bool bNormalizedFormat = PackedData.Format == ESparseVolumePackedDataFormat::Unorm8;
-	const EPixelFormat MultiCompFormat = GetMultiComponentFormat(PackedData.Format, NumActualComponents);
-	
-	if (MultiCompFormat == PF_Unknown)
+	// All source grid indices are INDEX_NONE, so nothing was selected for import
+	if (!bAnySourceGridIndicesValid)
 	{
-		// SVT_TODO error message about unsupported format
+		// SVT_TODO error message
 		return false;
 	}
 
@@ -291,72 +321,93 @@ bool ConvertOpenVDBToSparseVolumeTexture(
 	std::istream IStream(&StreamBuf);
 	openvdb::io::Stream Stream(IStream, false /*delayLoad*/);
 
-	// Check that at least one component has a valid grid assigned to it
+	// Check that the source grid indices are valid
 	openvdb::GridPtrVecPtr Grids = Stream.getGrids();
 	const size_t NumSourceGrids = Grids ? Grids->size() : 0;
-	if (NumSourceGrids <= PackedData.SourceGridIndex.X
-		&& NumSourceGrids <= PackedData.SourceGridIndex.Y
-		&& NumSourceGrids <= PackedData.SourceGridIndex.Z
-		&& NumSourceGrids <= PackedData.SourceGridIndex.W)
+	for (uint32 PackedDataIdx = 0; PackedDataIdx < NumPackedData; ++PackedDataIdx)
 	{
-		return false;
+		for (uint32 CompIdx = 0; CompIdx < 4; ++CompIdx)
+		{
+			const uint32 SourceGridIndex = PackedData[PackedDataIdx]->SourceGridIndex[CompIdx];
+			if (SourceGridIndex != INDEX_NONE && SourceGridIndex >= NumSourceGrids)
+			{
+				// SVT_TODO error message
+				return false;
+			}
+		}
 	}
 
-	FSparseVolumeAssetHeader& Header = *OutHeader;
-	Header.PackedDataAFormat = MultiCompFormat;
+	FSparseVolumeAssetHeader& Header = *OutResult->Header;
+	Header.PackedDataAFormat = MultiCompFormat[0];
+	Header.PackedDataBFormat = MultiCompFormat[1];
 	Header.SourceVolumeResolution = FIntVector::ZeroValue;
 	
 	FIntVector SmallestAABBMin = FIntVector(INT32_MAX);
 
 	// Compute per source grid data of up to 4 different grids (one per component)
-	TSharedPtr<IOpenVDBGridAdapterBase> GridAdapters[4]{};
-	float GridBackgroundValues[4]{};
-	float NormalizeScale[4]{ 1.0f, 1.0f, 1.0f, 1.0f };
-	float NormalizeBias[4]{};
-	for (uint32 CompIdx = 0; CompIdx < 4; ++CompIdx)
+	TArray<TSharedPtr<IOpenVDBGridAdapterBase>> UniqueGridAdapters;
+	UniqueGridAdapters.SetNum((int32)Grids->size());
+	TSharedPtr<IOpenVDBGridAdapterBase> GridAdapters[NumPackedData][4]{};
+	float GridBackgroundValues[NumPackedData][4]{};
+	float NormalizeScale[NumPackedData][4]{};
+	float NormalizeBias[NumPackedData][4]{};
+	for (uint32 PackedDataIdx = 0; PackedDataIdx < NumPackedData; ++PackedDataIdx)
 	{
-		if (PackedData.SourceGridIndex[CompIdx] == INDEX_NONE)
+		for (uint32 CompIdx = 0; CompIdx < 4; ++CompIdx)
 		{
-			GridAdapters[CompIdx] = nullptr;
-			continue;
-		}
+			NormalizeScale[PackedDataIdx][CompIdx] = 1.0f;
+			const uint32 SourceGridIndex = PackedData[PackedDataIdx]->SourceGridIndex[CompIdx];
+			const uint32 SourceComponentIndex = PackedData[PackedDataIdx]->SourceComponentIndex[CompIdx];
+			if (SourceGridIndex == INDEX_NONE)
+			{
+				continue;
+			}
 
-		openvdb::GridBase::Ptr GridBase = (*Grids)[PackedData.SourceGridIndex[CompIdx]];
-		GridAdapters[CompIdx] = CreateOpenVDBGridAdapter(GridBase);
-		if (!GridAdapters[CompIdx])
-		{
-			return false;
-		}
+			openvdb::GridBase::Ptr GridBase = (*Grids)[SourceGridIndex];
 
-		FOpenVDBData OVDBData = GetOpenVDBData(GridBase);
-		if (!IsOpenVDBDataValid(OVDBData, TEXT("")))
-		{
-			return false;
-		}
+			// Try to reuse adapters. Internally they use caching to accelerate read accesses, 
+			// so using three different adapters to access the three components of a single grid would be wasteful.
+			if (UniqueGridAdapters[SourceGridIndex] == nullptr)
+			{
+				UniqueGridAdapters[SourceGridIndex] = CreateOpenVDBGridAdapter(GridBase);
+				if (!UniqueGridAdapters[SourceGridIndex])
+				{
+					return false;
+				}
+			}
 
-		if (bOverrideActiveMinMax)
-		{
-			OVDBData.VolumeActiveAABBMin = ActiveMin;
-			OVDBData.VolumeActiveAABBMax = ActiveMax;
-			OVDBData.VolumeActiveDim = ActiveMax - ActiveMin;
-		}
+			GridAdapters[PackedDataIdx][CompIdx] = UniqueGridAdapters[SourceGridIndex];
 
-		Header.SourceVolumeResolution.X = FMath::Max(Header.SourceVolumeResolution.X, OVDBData.VolumeActiveDim.X);
-		Header.SourceVolumeResolution.Y = FMath::Max(Header.SourceVolumeResolution.Y, OVDBData.VolumeActiveDim.Y);
-		Header.SourceVolumeResolution.Z = FMath::Max(Header.SourceVolumeResolution.Z, OVDBData.VolumeActiveDim.Z);
-		SmallestAABBMin.X = FMath::Min(SmallestAABBMin.X, OVDBData.VolumeActiveAABBMin.X);
-		SmallestAABBMin.Y = FMath::Min(SmallestAABBMin.Y, OVDBData.VolumeActiveAABBMin.Y);
-		SmallestAABBMin.Z = FMath::Min(SmallestAABBMin.Z, OVDBData.VolumeActiveAABBMin.Z);
+			FOpenVDBData OVDBData = GetOpenVDBData(GridBase);
+			if (!IsOpenVDBDataValid(OVDBData, TEXT("")))
+			{
+				return false;
+			}
 
-		GridBackgroundValues[CompIdx] = GridAdapters[CompIdx]->GetBackgroundValue(PackedData.SourceComponentIndex[CompIdx]);
-		if (bNormalizedFormat && PackedData.bRemapInputForUnorm)
-		{
-			float MinVal = 0.0f;
-			float MaxVal = 0.0f;
-			GridAdapters[CompIdx]->GetMinMaxValue(PackedData.SourceComponentIndex[CompIdx], &MinVal, &MaxVal);
-			const float Diff = MaxVal - MinVal;
-			NormalizeScale[CompIdx] = MaxVal > SMALL_NUMBER ? (1.0f / Diff) : 1.0f;
-			NormalizeBias[CompIdx] = -MinVal * NormalizeScale[CompIdx];
+			if (bOverrideActiveMinMax)
+			{
+				OVDBData.VolumeActiveAABBMin = ActiveMin;
+				OVDBData.VolumeActiveAABBMax = ActiveMax;
+				OVDBData.VolumeActiveDim = ActiveMax - ActiveMin;
+			}
+
+			Header.SourceVolumeResolution.X = FMath::Max(Header.SourceVolumeResolution.X, OVDBData.VolumeActiveDim.X);
+			Header.SourceVolumeResolution.Y = FMath::Max(Header.SourceVolumeResolution.Y, OVDBData.VolumeActiveDim.Y);
+			Header.SourceVolumeResolution.Z = FMath::Max(Header.SourceVolumeResolution.Z, OVDBData.VolumeActiveDim.Z);
+			SmallestAABBMin.X = FMath::Min(SmallestAABBMin.X, OVDBData.VolumeActiveAABBMin.X);
+			SmallestAABBMin.Y = FMath::Min(SmallestAABBMin.Y, OVDBData.VolumeActiveAABBMin.Y);
+			SmallestAABBMin.Z = FMath::Min(SmallestAABBMin.Z, OVDBData.VolumeActiveAABBMin.Z);
+
+			GridBackgroundValues[PackedDataIdx][CompIdx] = GridAdapters[PackedDataIdx][CompIdx]->GetBackgroundValue(SourceComponentIndex);
+			if (bNormalizedFormat[PackedDataIdx] && PackedData[PackedDataIdx]->bRemapInputForUnorm)
+			{
+				float MinVal = 0.0f;
+				float MaxVal = 0.0f;
+				GridAdapters[PackedDataIdx][CompIdx]->GetMinMaxValue(SourceComponentIndex, &MinVal, &MaxVal);
+				const float Diff = MaxVal - MinVal;
+				NormalizeScale[PackedDataIdx][CompIdx] = MaxVal > SMALL_NUMBER ? (1.0f / Diff) : 1.0f;
+				NormalizeBias[PackedDataIdx][CompIdx] = -MinVal * NormalizeScale[PackedDataIdx][CompIdx];
+			}
 		}
 	}
 	
@@ -367,9 +418,6 @@ bool ConvertOpenVDBToSparseVolumeTexture(
 		FMath::DivideAndRoundUp(Header.SourceVolumeResolution.Z, SPARSE_VOLUME_TILE_RES));
 	Header.TileDataVolumeResolution = FIntVector::ZeroValue;	// unknown for now
 
-	const uint32 VolumeTileDataBytes = GPixelFormats[(SIZE_T)MultiCompFormat].Get3DImageSizeInBytes(SPARSE_VOLUME_TILE_RES, SPARSE_VOLUME_TILE_RES, SPARSE_VOLUME_TILE_RES);
-	const uint32 FormatSize = (uint32)GPixelFormats[(SIZE_T)MultiCompFormat].BlockBytes;
-	const uint32 SingleComponentFormatSize = FormatSize / NumActualComponents;
 	uint32 NumAllocatedPages = 0;
 
 	// Allocate some memory for temp data (worst case)
@@ -387,19 +435,26 @@ bool ConvertOpenVDBToSparseVolumeTexture(
 			{
 				bool bHasAnyData = false;
 
-				for (int32_t z = 0; z < SPARSE_VOLUME_TILE_RES; ++z)
+				for (uint32 PackedDataIdx = 0; PackedDataIdx < NumPackedData; ++PackedDataIdx)
 				{
-					for (int32_t y = 0; y < SPARSE_VOLUME_TILE_RES; ++y)
+					if (!bHasValidSourceGrids[PackedDataIdx])
+					{
+						continue;
+					}
+					for (int32_t z = 0; z < SPARSE_VOLUME_TILE_RES; ++z)
 					{
-						for (int32_t x = 0; x < SPARSE_VOLUME_TILE_RES; ++x)
+						for (int32_t y = 0; y < SPARSE_VOLUME_TILE_RES; ++y)
 						{
-							for (uint32 CompIdx = 0; CompIdx < NumActualComponents; ++CompIdx)
+							for (int32_t x = 0; x < SPARSE_VOLUME_TILE_RES; ++x)
 							{
-								if (PackedData.SourceGridIndex[CompIdx] != INDEX_NONE)
+								for (uint32 CompIdx = 0; CompIdx < NumActualComponents[PackedDataIdx]; ++CompIdx)
 								{
-									FVector VoxelCoord = FVector(SmallestAABBMin) + FVector(PageX, PageY, PageZ) * SPARSE_VOLUME_TILE_RES + FVector(x, y, z);	// This assumes sampling outside the boundary returns a default value
-									float VoxelValue = GridAdapters[CompIdx]->Sample(openvdb::Coord(VoxelCoord.X, VoxelCoord.Y, VoxelCoord.Z), PackedData.SourceComponentIndex[CompIdx]);
-									bHasAnyData |= VoxelValue != GridBackgroundValues[CompIdx];
+									if (PackedData[PackedDataIdx]->SourceGridIndex[CompIdx] != INDEX_NONE)
+									{
+										FVector VoxelCoord = FVector(SmallestAABBMin) + FVector(PageX, PageY, PageZ) * SPARSE_VOLUME_TILE_RES + FVector(x, y, z);	// This assumes sampling outside the boundary returns a default value
+										float VoxelValue = GridAdapters[PackedDataIdx][CompIdx]->Sample(openvdb::Coord(VoxelCoord.X, VoxelCoord.Y, VoxelCoord.Z), PackedData[PackedDataIdx]->SourceComponentIndex[CompIdx]);
+										bHasAnyData |= VoxelValue != GridBackgroundValues[PackedDataIdx][CompIdx];
+									}
 								}
 							}
 						}
@@ -432,10 +487,11 @@ bool ConvertOpenVDBToSparseVolumeTexture(
 	Header.TileDataVolumeResolution = Header.TileDataVolumeResolution * SPARSE_VOLUME_TILE_RES;
 
 	// Initialise the SparseVolumeTexture page and tile.
-	OutDensityPage->SetNumZeroed(Header.PageTableVolumeResolution.X * Header.PageTableVolumeResolution.Y * Header.PageTableVolumeResolution.Z);
-	OutDensityData->SetNumZeroed(Header.TileDataVolumeResolution.X * Header.TileDataVolumeResolution.Y * Header.TileDataVolumeResolution.Z * FormatSize);
-	uint32* DensityPagePtr = OutDensityPage->GetData();
-	uint8* DensityDataPtr = OutDensityData->GetData();
+	OutResult->PageTable->SetNumZeroed(Header.PageTableVolumeResolution.X * Header.PageTableVolumeResolution.Y * Header.PageTableVolumeResolution.Z);
+	OutResult->PhysicalTileDataA->SetNumZeroed(Header.TileDataVolumeResolution.X * Header.TileDataVolumeResolution.Y * Header.TileDataVolumeResolution.Z * FormatSize[0] * (bHasValidSourceGrids[0] ? 1 : 0));
+	OutResult->PhysicalTileDataB->SetNumZeroed(Header.TileDataVolumeResolution.X * Header.TileDataVolumeResolution.Y * Header.TileDataVolumeResolution.Z * FormatSize[1] * (bHasValidSourceGrids[1] ? 1 : 0));
+	uint32* PageTablePtr = OutResult->PageTable->GetData();
+	uint8* PhysicalTileDataPtrs[] = { OutResult->PhysicalTileDataA->GetData(), OutResult->PhysicalTileDataB->GetData() };
 
 	FIntVector DestinationTileCoord = FIntVector::ZeroValue;
 	auto GoToNextTileCoord = [&]()
@@ -456,7 +512,7 @@ bool ConvertOpenVDBToSparseVolumeTexture(
 	// Add an empty tile is needed, reserve slot at coord 0
 	if (bAnyEmptyPageExists)
 	{
-		// DensityPage is all cleared to zero, simply skip a tile
+		// PageTable is all cleared to zero, simply skip a tile
 		GoToNextTileCoord();
 	}
 
@@ -471,7 +527,7 @@ bool ConvertOpenVDBToSparseVolumeTexture(
 		uint32 DestinationTileCoord32bit = (DestinationTileCoord.X & 0x7FF) | ((DestinationTileCoord.Y & 0x7FF) << 11) | ((DestinationTileCoord.Z & 0x3FF) << 22);
 
 		// Setup the page table entry
-		DensityPagePtr
+		PageTablePtr
 			[
 				PageCoordToSplat.Z * Header.PageTableVolumeResolution.X * Header.PageTableVolumeResolution.Y +
 				PageCoordToSplat.Y * Header.PageTableVolumeResolution.X +
@@ -479,48 +535,56 @@ bool ConvertOpenVDBToSparseVolumeTexture(
 			] = DestinationTileCoord32bit;
 
 		// Now copy the tile data from the source to the page atlas
-		for (int32_t z = 0; z < SPARSE_VOLUME_TILE_RES; ++z)
+		for (uint32 PackedDataIdx = 0; PackedDataIdx < NumPackedData; ++PackedDataIdx)
 		{
-			for (int32_t y = 0; y < SPARSE_VOLUME_TILE_RES; ++y)
+			if (!bHasValidSourceGrids[PackedDataIdx])
+			{
+				continue;
+			}
+			uint8* PhysicalTileDataPtr = PhysicalTileDataPtrs[PackedDataIdx];
+			for (int32_t z = 0; z < SPARSE_VOLUME_TILE_RES; ++z)
 			{
-				for (int32_t x = 0; x < SPARSE_VOLUME_TILE_RES; ++x)
+				for (int32_t y = 0; y < SPARSE_VOLUME_TILE_RES; ++y)
 				{
-					for (uint32 CompIdx = 0; CompIdx < NumActualComponents; ++CompIdx)
+					for (int32_t x = 0; x < SPARSE_VOLUME_TILE_RES; ++x)
 					{
-						float VoxelValue = 0.0f;
-						float VoxelValueNormalized = 0.0f;
-						if (PackedData.SourceGridIndex[CompIdx] != INDEX_NONE)
+						for (uint32 CompIdx = 0; CompIdx < NumActualComponents[PackedDataIdx]; ++CompIdx)
 						{
-							FVector VoxelCoord = FVector(SmallestAABBMin) + FVector(PageCoordToSplat * SPARSE_VOLUME_TILE_RES) + FVector(x, y, z);	// This assumes sampling outside the boundary returns a default value
-							VoxelValue = GridAdapters[CompIdx]->Sample(openvdb::Coord(VoxelCoord.X, VoxelCoord.Y, VoxelCoord.Z), PackedData.SourceComponentIndex[CompIdx]);
-							VoxelValueNormalized = FMath::Clamp(VoxelValue * NormalizeScale[CompIdx] + NormalizeBias[CompIdx], 0.0f, 1.0f);
-						}
+							float VoxelValue = 0.0f;
+							float VoxelValueNormalized = 0.0f;
+							if (PackedData[PackedDataIdx]->SourceGridIndex[CompIdx] != INDEX_NONE)
+							{
+								FVector VoxelCoord = FVector(SmallestAABBMin) + FVector(PageCoordToSplat * SPARSE_VOLUME_TILE_RES) + FVector(x, y, z);	// This assumes sampling outside the boundary returns a default value
+								VoxelValue = GridAdapters[PackedDataIdx][CompIdx]->Sample(openvdb::Coord(VoxelCoord.X, VoxelCoord.Y, VoxelCoord.Z), PackedData[PackedDataIdx]->SourceComponentIndex[CompIdx]);
+								VoxelValueNormalized = FMath::Clamp(VoxelValue * NormalizeScale[PackedDataIdx][CompIdx] + NormalizeBias[PackedDataIdx][CompIdx], 0.0f, 1.0f);
+							}
 
-						const SIZE_T DstCoord = ((DestinationTileCoord.Z * SPARSE_VOLUME_TILE_RES + z) * Header.TileDataVolumeResolution.X * Header.TileDataVolumeResolution.Y +
-							(DestinationTileCoord.Y * SPARSE_VOLUME_TILE_RES + y) * Header.TileDataVolumeResolution.X +
-							(DestinationTileCoord.X * SPARSE_VOLUME_TILE_RES + x)) * FormatSize + CompIdx * SingleComponentFormatSize;
+							const SIZE_T DstCoord = ((DestinationTileCoord.Z * SPARSE_VOLUME_TILE_RES + z) * Header.TileDataVolumeResolution.X * Header.TileDataVolumeResolution.Y +
+								(DestinationTileCoord.Y * SPARSE_VOLUME_TILE_RES + y) * Header.TileDataVolumeResolution.X +
+								(DestinationTileCoord.X * SPARSE_VOLUME_TILE_RES + x)) * FormatSize[PackedDataIdx] + CompIdx * SingleComponentFormatSize[PackedDataIdx];
 
-						switch (PackedData.Format)
-						{
-						case ESparseVolumePackedDataFormat::Unorm8:
-						{
-							DensityDataPtr[DstCoord] = uint8(VoxelValueNormalized * 255.0f);
-							break;
-						}
-						case ESparseVolumePackedDataFormat::Float16:
-						{
-							const uint16 VoxelValue16FEncoded = FFloat16(VoxelValue).Encoded;
-							*((uint16*)(&DensityDataPtr[DstCoord])) = VoxelValue16FEncoded;
-							break;
-						}
-						case ESparseVolumePackedDataFormat::Float32:
-						{
-							*((float*)(&DensityDataPtr[DstCoord])) = VoxelValue;
-							break;
-						}
-						default:
-							checkNoEntry();
-							break;
+							switch (PackedData[PackedDataIdx]->Format)
+							{
+							case ESparseVolumePackedDataFormat::Unorm8:
+							{
+								PhysicalTileDataPtr[DstCoord] = uint8(VoxelValueNormalized * 255.0f);
+								break;
+							}
+							case ESparseVolumePackedDataFormat::Float16:
+							{
+								const uint16 VoxelValue16FEncoded = FFloat16(VoxelValue).Encoded;
+								*((uint16*)(&PhysicalTileDataPtr[DstCoord])) = VoxelValue16FEncoded;
+								break;
+							}
+							case ESparseVolumePackedDataFormat::Float32:
+							{
+								*((float*)(&PhysicalTileDataPtr[DstCoord])) = VoxelValue;
+								break;
+							}
+							default:
+								checkNoEntry();
+								break;
+							}
 						}
 					}
 				}
diff --git a/Engine/Source/Editor/SparseVolumeTexture/Private/SparseVolumeTextureOpenVDBUtility.h b/Engine/Source/Editor/SparseVolumeTexture/Private/SparseVolumeTextureOpenVDBUtility.h
--- a/Engine/Source/Editor/SparseVolumeTexture/Private/SparseVolumeTextureOpenVDBUtility.h
+++ b/Engine/Source/Editor/SparseVolumeTexture/Private/SparseVolumeTextureOpenVDBUtility.h
@@ -47,10 +47,9 @@ bool GetOpenVDBGridInfo(TArray<uint8>& SourceFile, TArray<FOpenVDBGridInfo>* Out
 
 bool ConvertOpenVDBToSparseVolumeTexture(
 	TArray<uint8>& SourceFile,
-	struct FSparseVolumeRawSourcePackedData& PackedData,
-	struct FSparseVolumeAssetHeader* OutHeader,
-	TArray<uint32>* OutDensityPage,
-	TArray<uint8>* OutDensityData,
+	struct FSparseVolumeRawSourcePackedData& PackedDataA,
+	struct FSparseVolumeRawSourcePackedData& PackedDataB,
+	struct FOpenVDBToSVTConversionResult* OutResult,
 	bool bOverrideActiveMinMax,
 	FVector ActiveMin,
 	FVector ActiveMax);
diff --git a/Engine/Source/Runtime/Engine/Classes/SparseVolumeTexture/SparseVolumeTexture.h b/Engine/Source/Runtime/Engine/Classes/SparseVolumeTexture/SparseVolumeTexture.h
--- a/Engine/Source/Runtime/Engine/Classes/SparseVolumeTexture/SparseVolumeTexture.h
+++ b/Engine/Source/Runtime/Engine/Classes/SparseVolumeTexture/SparseVolumeTexture.h
@@ -17,12 +17,19 @@ namespace UE { namespace Shader	{ enum class EValueType : uint8; } }
 
 #define SPARSE_VOLUME_TILE_RES	16
 
-DECLARE_DELEGATE_RetVal_EightParams(bool, FConvertOpenVDBToSparseVolumeTextureDelegate, 
+struct ENGINE_API FOpenVDBToSVTConversionResult
+{
+	struct FSparseVolumeAssetHeader* Header;
+	TArray<uint32>* PageTable;
+	TArray<uint8>* PhysicalTileDataA;
+	TArray<uint8>* PhysicalTileDataB;
+};
+
+DECLARE_DELEGATE_RetVal_SevenParams(bool, FConvertOpenVDBToSparseVolumeTextureDelegate, 
 	TArray<uint8>& SourceFile,
-	struct FSparseVolumeRawSourcePackedData& PackedData,
-	struct FSparseVolumeAssetHeader* OutHeader,
-	TArray<uint32>* OutDensityPage,
-	TArray<uint8>* OutDensityData,
+	struct FSparseVolumeRawSourcePackedData& PackedDataA,
+	struct FSparseVolumeRawSourcePackedData& PackedDataB,
+	FOpenVDBToSVTConversionResult* OutResult,
 	bool bOverrideActiveMinMax,
 	FVector ActiveMin,
 	FVector ActiveMax);
@@ -42,6 +49,7 @@ struct ENGINE_API FSparseVolumeAssetHeader
 	FIntVector3 TileDataVolumeResolution;
 	FIntVector3 SourceVolumeResolution;
 	EPixelFormat PackedDataAFormat;
+	EPixelFormat PackedDataBFormat;
 
 	// The current data format version for the header.
 	static const uint32 kVersion = 0;
@@ -70,10 +78,13 @@ struct ENGINE_API FSparseVolumeRawSourcePackedData
 	bool bRemapInputForUnorm; // Maps the input from its minimum and maximum value into the [0-1] range. Clamps to [0-1] otherwise.
 };
 
+ENGINE_API FArchive& operator<<(FArchive& Ar, FSparseVolumeRawSourcePackedData& PackedData);
+
 // The structure represent the source asset in high quality. It is used to cook the runtime data
 struct ENGINE_API FSparseVolumeRawSource
 {
 	FSparseVolumeRawSourcePackedData PackedDataA;
+	FSparseVolumeRawSourcePackedData PackedDataB;
 	TArray<uint8> SourceAssetFile;
 
 	// The current data format version for the raw source data.
@@ -94,8 +105,9 @@ struct ENGINE_API FSparseVolumeRawSource
 struct ENGINE_API FSparseVolumeTextureRuntime
 {
 	FSparseVolumeAssetHeader	Header;
-	TArray<uint32>				DensityPage;
-	TArray<uint8>				DensityData;
+	TArray<uint32>				PageTable;
+	TArray<uint8>				PhysicalTileDataA;
+	TArray<uint8>				PhysicalTileDataB;
 
 	void SetAsDefaultTexture();
 
@@ -130,13 +142,17 @@ public:
 		return SparseVolumeTextureRuntime->Header;
 	}
 
-	FTextureRHIRef GetTileDataTextureRHI() const
+	FTextureRHIRef GetPhysicalTileDataATextureRHI() const
+	{
+		return PhysicalTileDataATextureRHI;
+	}
+	FTextureRHIRef GetPhysicalTileDataBTextureRHI() const
 	{
-		return  TileDataTextureRHI;
+		return PhysicalTileDataBTextureRHI;
 	}
 	FTextureRHIRef GetPageTableTextureRHI() const
 	{
-		return  PageTableTextureRHI;
+		return PageTableTextureRHI;
 	}
 
 	virtual void InitRHI() override;
@@ -147,7 +163,8 @@ private:
 	FSparseVolumeTextureRuntime*		SparseVolumeTextureRuntime;
 
 	FTextureRHIRef						PageTableTextureRHI;
-	FTextureRHIRef						TileDataTextureRHI;
+	FTextureRHIRef						PhysicalTileDataATextureRHI;
+	FTextureRHIRef						PhysicalTileDataBTextureRHI;
 };
 
 
diff --git a/Engine/Source/Runtime/Engine/Private/Materials/HLSLMaterialTranslator.cpp b/Engine/Source/Runtime/Engine/Private/Materials/HLSLMaterialTranslator.cpp
--- a/Engine/Source/Runtime/Engine/Private/Materials/HLSLMaterialTranslator.cpp
+++ b/Engine/Source/Runtime/Engine/Private/Materials/HLSLMaterialTranslator.cpp
@@ -12405,15 +12405,16 @@ int32 FHLSLMaterialTranslator::SparseVolumeTextureSample(int32 TextureIndex, int
 		int32 VirtualTextureIndex = UniformTextureExpressions[(uint32)EMaterialTextureParameterType::Volume].Find(TextureUniformExpression);
 		check(UniformTextureExpressions[(uint32)EMaterialTextureParameterType::Volume].IsValidIndex(VirtualTextureIndex));
 
-		check(TextureUniformExpression->GetTextureLayerIndex() != INDEX_NONE);
-		if (TextureUniformExpression->GetTextureLayerIndex() == 0)
+		const int32 TextureLayerIndex = TextureUniformExpression->GetTextureLayerIndex();
+		check(TextureLayerIndex != INDEX_NONE);
+		if (TextureLayerIndex == 0)
 		{
 			FString SampleCode = FString::Printf(TEXT("%s.Load(int4(int3(%s), 0)).x"), *TextureName, *GetParameterCode(CoordinateIndex));
 			AddEstimatedTextureSample();
 			// We ignore derivative for now, see IsAnalyticDerivEnabled() and IsDerivativeValid(UvDerivativeStatus)
 			SamplingCodeIndex = AddCodeChunk(MCT_UInt1, *SampleCode);
 		}
-		else if (TextureUniformExpression->GetTextureLayerIndex() == 1)
+		else if (TextureLayerIndex == 1 || TextureLayerIndex == 2)
 		{
 			FString SampleCode = FString::Printf(TEXT("%s.Load(int4(int3(%s), 0))"), *TextureName, *GetParameterCode(CoordinateIndex));
 			AddEstimatedTextureSample();
diff --git a/Engine/Source/Runtime/Engine/Private/Materials/MaterialExpressions.cpp b/Engine/Source/Runtime/Engine/Private/Materials/MaterialExpressions.cpp
--- a/Engine/Source/Runtime/Engine/Private/Materials/MaterialExpressions.cpp
+++ b/Engine/Source/Runtime/Engine/Private/Materials/MaterialExpressions.cpp
@@ -25234,9 +25234,9 @@ UMaterialExpressionSparseVolumeTextureSample::UMaterialExpressionSparseVolumeTex
 #endif
 
 #if WITH_EDITOR
-	// SVT_TODO update that when attribute setup has been clarified
 	Outputs.Reset();
-	Outputs.Add(FExpressionOutput(TEXT("Density"), 1, 1, 1, 1, 0));
+	Outputs.Add(FExpressionOutput(TEXT("PackedData A"), 1, 1, 1, 1, 1));
+	Outputs.Add(FExpressionOutput(TEXT("PackedData B"), 1, 1, 1, 1, 1));
 	bShowOutputNameOnPin = true;
 	bShowMaskColorsOnPin = false;
 #endif
@@ -25273,7 +25273,8 @@ uint32 UMaterialExpressionSparseVolumeTextureSample::GetOutputType(int32 OutputI
 	switch (OutputIndex)
 	{
 	case 0:
-		return MCT_Float;	// Density defined in constructor for now
+	case 1:
+		return MCT_Float;	// PackedData A and B defined in constructor for now
 		break;
 	}
 
@@ -25313,7 +25314,7 @@ int32 UMaterialExpressionSparseVolumeTextureSample::Compile(class FMaterialCompi
 		}
 
 		int32 PhysicalTileTextureReferenceIndex = INDEX_NONE;
-		int32 PhysicalTileSubTextureIndex = 1;
+		int32 PhysicalTileSubTextureIndex = 1 + OutputIndex; // index 0 is used by the page table texture, so we start at 1
 		int32 PhysicalTileTextureParamCodeChunkIndex = INDEX_NONE;
 		if (bIsParameter)
 		{
@@ -25345,6 +25346,9 @@ int32 UMaterialExpressionSparseVolumeTextureSample::Compile(class FMaterialCompi
 		int32 UVWPageTable = Compiler->Mul(UVW, UniformCodeChunkIndex[ESparseVolumeTexture_PhysicalUVToPageUV]);
 		int32 CoordPageTable = Compiler->Mul(UVWPageTable, UniformCodeChunkIndex[ESparseVolumeTexture_PageTableSize]);
 
+		// Code chunks have the same ID when they map to the same string of code. PackedPhysicalTileCoord will automatically end up being the same code chunk used
+		// when calling Compile() for Output 0 (PackedDataA) and Output 1 (PackedDataB). This is perfect for this use case: 
+		// PackedPhysicalTileCoord will end up being evaluated only once for both outputs, so they both share a single page table lookup.
 		int32 PackedPhysicalTileCoord = Compiler->SparseVolumeTextureSample(PageTableTextureParamCodeChunkIndex, CoordPageTable);
 
 		int32 CoordVolume = Compiler->Mul(CoordPageTable, UniformCodeChunkIndex[ESparseVolumeTexture_TileSize]);
diff --git a/Engine/Source/Runtime/Engine/Private/Materials/MaterialUniformExpressions.cpp b/Engine/Source/Runtime/Engine/Private/Materials/MaterialUniformExpressions.cpp
--- a/Engine/Source/Runtime/Engine/Private/Materials/MaterialUniformExpressions.cpp
+++ b/Engine/Source/Runtime/Engine/Private/Materials/MaterialUniformExpressions.cpp
@@ -496,6 +496,14 @@ FShaderParametersMetadata* FUniformExpressionSet::CreateBufferStruct()
 				new(Members) FShaderParametersMetadata::FMember(*VolumeTextureSamplerNames[i], TEXT("SamplerState"), __LINE__, NextMemberOffset, UBMT_SAMPLER, EShaderPrecisionModifier::Float, 1, 1, 0, NULL);
 				NextMemberOffset += SHADER_PARAMETER_POINTER_ALIGNMENT;
 			}
+			else if (Parameter.VirtualTextureLayerIndex == 2)
+			{
+				check((NextMemberOffset % SHADER_PARAMETER_POINTER_ALIGNMENT) == 0);
+				new(Members) FShaderParametersMetadata::FMember(*VolumeTextureNames[i], TEXT("Texture3D"), __LINE__, NextMemberOffset, UBMT_TEXTURE, EShaderPrecisionModifier::Float, 1, 1, 0, NULL);
+				NextMemberOffset += SHADER_PARAMETER_POINTER_ALIGNMENT;
+				new(Members) FShaderParametersMetadata::FMember(*VolumeTextureSamplerNames[i], TEXT("SamplerState"), __LINE__, NextMemberOffset, UBMT_SAMPLER, EShaderPrecisionModifier::Float, 1, 1, 0, NULL);
+				NextMemberOffset += SHADER_PARAMETER_POINTER_ALIGNMENT;
+			}
 		}
 
 
@@ -1192,8 +1200,8 @@ void FUniformExpressionSet::FillUniformBuffer(const FMaterialRenderContext& Mate
 				{
 					const int32 LayerIndex = Parameter.VirtualTextureLayerIndex;
 					// 0 is page table
-					// 1 is density
-					// SVT_TODO update that when attribute setup has been clarified
+					// 1 is tile data A
+					// 2 is tile data B
 
 					const FSparseVolumeTextureSceneProxy* SVTextureProxy = SVTexture->GetSparseVolumeTextureSceneProxy();
 					if (SVTextureProxy)
@@ -1205,9 +1213,13 @@ void FUniformExpressionSet::FillUniformBuffer(const FMaterialRenderContext& Mate
 						}
 						else if(LayerIndex == 1)
 						{
-
-							check(SVTextureProxy->GetTileDataTextureRHI());
-							*ResourceTableTexturePtr = SVTextureProxy->GetTileDataTextureRHI();
+							FTextureRHIRef TileDataATexture = SVTextureProxy->GetPhysicalTileDataATextureRHI();
+							*ResourceTableTexturePtr = TileDataATexture ? TileDataATexture : GBlackVolumeTexture->TextureRHI;
+						}
+						else if (LayerIndex == 2)
+						{
+							FTextureRHIRef TileDataBTexture = SVTextureProxy->GetPhysicalTileDataBTextureRHI();
+							*ResourceTableTexturePtr = TileDataBTexture ? TileDataBTexture : GBlackVolumeTexture->TextureRHI;
 						}
 						else
 						{
diff --git a/Engine/Source/Runtime/Engine/Private/SparseVolumeTexture/SparseVolumeTexture.cpp b/Engine/Source/Runtime/Engine/Private/SparseVolumeTexture/SparseVolumeTexture.cpp
--- a/Engine/Source/Runtime/Engine/Private/SparseVolumeTexture/SparseVolumeTexture.cpp
+++ b/Engine/Source/Runtime/Engine/Private/SparseVolumeTexture/SparseVolumeTexture.cpp
@@ -44,12 +44,21 @@ void FSparseVolumeAssetHeader::Serialize(FArchive& Ar)
 		Ar << TileDataVolumeResolution;
 		Ar << SourceVolumeResolution;
 
-		uint32 FormatAsUint = (uint32)PackedDataAFormat;
+		uint32 FormatAsUint;
+		
+		FormatAsUint = (uint32)PackedDataAFormat;
 		Ar << FormatAsUint;
 		if (Ar.IsLoading())
 		{
 			PackedDataAFormat = static_cast<EPixelFormat>(FormatAsUint);
 		}
+		
+		FormatAsUint = (uint32)PackedDataBFormat;
+		Ar << FormatAsUint;
+		if (Ar.IsLoading())
+		{
+			PackedDataBFormat = static_cast<EPixelFormat>(FormatAsUint);
+		}
 	}
 	else
 	{
@@ -58,22 +67,29 @@ void FSparseVolumeAssetHeader::Serialize(FArchive& Ar)
 	}
 }
 
+FArchive& operator<<(FArchive& Ar, FSparseVolumeRawSourcePackedData& PackedData)
+{
+	uint32 FormatAsUint = (uint32)PackedData.Format;
+	Ar << FormatAsUint;
+	if (Ar.IsLoading())
+	{
+		PackedData.Format = static_cast<ESparseVolumePackedDataFormat>(FormatAsUint);
+	}
+	Ar << PackedData.SourceGridIndex;
+	Ar << PackedData.SourceComponentIndex;
+	Ar << PackedData.bRemapInputForUnorm;
+
+	return Ar;
+}
+
 void FSparseVolumeRawSource::Serialize(FArchive& Ar)
 {
 	Ar << Version;
 
 	if (Version == 0)
 	{
-		uint32 FormatAsUint = (uint32)PackedDataA.Format;
-		Ar << FormatAsUint;
-		if (Ar.IsLoading())
-		{
-			PackedDataA.Format = static_cast<ESparseVolumePackedDataFormat>(FormatAsUint);
-		}
-		Ar << PackedDataA.SourceGridIndex;
-		Ar << PackedDataA.SourceComponentIndex;
-		Ar << PackedDataA.bRemapInputForUnorm;
-
+		Ar << PackedDataA;
+		Ar << PackedDataB;
 		Ar << SourceAssetFile;
 	}
 	else
@@ -91,8 +107,9 @@ void FSparseVolumeTextureRuntime::Serialize(FArchive& Ar)
 
 	if (Version == 0)
 	{
-		Ar << DensityPage;
-		Ar << DensityData;
+		Ar << PageTable;
+		Ar << PhysicalTileDataA;
+		Ar << PhysicalTileDataB;
 	}
 	else
 	{
@@ -104,8 +121,9 @@ void FSparseVolumeTextureRuntime::Serialize(FArchive& Ar)
 void FSparseVolumeTextureRuntime::SetAsDefaultTexture()
 {
 	const uint32 VolumeSize = 1;
-	DensityPage.SetNumZeroed(VolumeSize * VolumeSize * VolumeSize);
-	DensityData.SetNumZeroed(VolumeSize * VolumeSize * VolumeSize);
+	PageTable.SetNumZeroed(VolumeSize * VolumeSize * VolumeSize);
+	PhysicalTileDataA.SetNumZeroed(VolumeSize * VolumeSize * VolumeSize);
+	PhysicalTileDataB.SetNumZeroed(VolumeSize * VolumeSize * VolumeSize);
 }
 
 ////////////////////////////////////////////////////////////////////////////////////////////////
@@ -246,13 +264,18 @@ void UStaticSparseVolumeTexture::ConvertRawSourceDataToSparseVolumeTextureRuntim
 		FSparseVolumeRawSource SparseVolumeRawSource;
 		SparseVolumeRawSource.Serialize(RawDataArchiveReader);
 
+		FOpenVDBToSVTConversionResult SVTResult;
+		SVTResult.Header = &StaticFrame.SparseVolumeTextureRuntime.Header;
+		SVTResult.PageTable = &StaticFrame.SparseVolumeTextureRuntime.PageTable;
+		SVTResult.PhysicalTileDataA = &StaticFrame.SparseVolumeTextureRuntime.PhysicalTileDataA;
+		SVTResult.PhysicalTileDataB = &StaticFrame.SparseVolumeTextureRuntime.PhysicalTileDataB;
+
 		// Then, cook the runtime data and serialize it out.
 		const bool bSuccess = ConvertOpenVDBToSparseVolumeTextureDelegate.IsBound() && ConvertOpenVDBToSparseVolumeTextureDelegate.Execute(
 			SparseVolumeRawSource.SourceAssetFile,
 			SparseVolumeRawSource.PackedDataA,
-			&StaticFrame.SparseVolumeTextureRuntime.Header,
-			&StaticFrame.SparseVolumeTextureRuntime.DensityPage,
-			&StaticFrame.SparseVolumeTextureRuntime.DensityData,
+			SparseVolumeRawSource.PackedDataB,
+			&SVTResult,
 			false, FVector::Zero(), FVector::Zero());
 		ensure(bSuccess);
 
@@ -385,7 +408,8 @@ FSparseVolumeTextureSceneProxy::FSparseVolumeTextureSceneProxy()
 	: FRenderResource()
 	, SparseVolumeTextureRuntime(nullptr)
 	, PageTableTextureRHI(nullptr)
-	, TileDataTextureRHI(nullptr)
+	, PhysicalTileDataATextureRHI(nullptr)
+	, PhysicalTileDataBTextureRHI(nullptr)
 {
 }
 
@@ -413,23 +437,43 @@ void FSparseVolumeTextureSceneProxy::InitRHI()
 
 		const int32 FormatSize = GPixelFormats[PageEntryFormat].BlockBytes;
 		const FUpdateTextureRegion3D UpdateRegion(0, 0, 0, 0, 0, 0, PageTableVolumeResolution.X, PageTableVolumeResolution.Y, PageTableVolumeResolution.Z);
-		RHIUpdateTexture3D(PageTableTextureRHI, 0, UpdateRegion, UpdateRegion.Width * FormatSize, UpdateRegion.Width * UpdateRegion.Height * FormatSize, (const uint8*)SparseVolumeTextureRuntime->DensityPage.GetData());
+		RHIUpdateTexture3D(PageTableTextureRHI, 0, UpdateRegion, UpdateRegion.Width * FormatSize, UpdateRegion.Width * UpdateRegion.Height * FormatSize, (const uint8*)SparseVolumeTextureRuntime->PageTable.GetData());
 	}
 
 	// Tile data
 	{
-		EPixelFormat VoxelFormat = SparseVolumeTextureRuntime->Header.PackedDataAFormat;
 		FIntVector3 TileDataVolumeResolution = SparseVolumeTextureRuntime->Header.TileDataVolumeResolution;
-		const FRHITextureCreateDesc Desc =
-			FRHITextureCreateDesc::Create3D(TEXT("SparseVolumeTexture.TileData.RHITexture"),
-				TileDataVolumeResolution.X, TileDataVolumeResolution.Y, TileDataVolumeResolution.Z, VoxelFormat)
-			.SetFlags(ETextureCreateFlags::ShaderResource);
+		EPixelFormat VoxelFormatA = SparseVolumeTextureRuntime->Header.PackedDataAFormat;
+		EPixelFormat VoxelFormatB = SparseVolumeTextureRuntime->Header.PackedDataBFormat;
+		const FUpdateTextureRegion3D UpdateRegion(0, 0, 0, 0, 0, 0, TileDataVolumeResolution.X, TileDataVolumeResolution.Y, TileDataVolumeResolution.Z);
 
-		TileDataTextureRHI = RHICreateTexture(Desc);
+		// A
+		if (VoxelFormatA != PF_Unknown)
+		{
+			const FRHITextureCreateDesc Desc =
+				FRHITextureCreateDesc::Create3D(TEXT("SparseVolumeTexture.PhysicalTileDataA.RHITexture"),
+					TileDataVolumeResolution.X, TileDataVolumeResolution.Y, TileDataVolumeResolution.Z, VoxelFormatA)
+				.SetFlags(ETextureCreateFlags::ShaderResource);
 
-		const int32 FormatSize = GPixelFormats[VoxelFormat].BlockBytes;
-		const FUpdateTextureRegion3D UpdateRegion(0, 0, 0, 0, 0, 0, TileDataVolumeResolution.X, TileDataVolumeResolution.Y, TileDataVolumeResolution.Z);
-		RHIUpdateTexture3D(TileDataTextureRHI, 0, UpdateRegion, UpdateRegion.Width * FormatSize, UpdateRegion.Width * UpdateRegion.Height * FormatSize, (const uint8*)SparseVolumeTextureRuntime->DensityData.GetData());
+			PhysicalTileDataATextureRHI = RHICreateTexture(Desc);
+
+			const int32 FormatSize = GPixelFormats[VoxelFormatA].BlockBytes;
+			RHIUpdateTexture3D(PhysicalTileDataATextureRHI, 0, UpdateRegion, UpdateRegion.Width * FormatSize, UpdateRegion.Width * UpdateRegion.Height * FormatSize, (const uint8*)SparseVolumeTextureRuntime->PhysicalTileDataA.GetData());
+		}
+		
+		// B
+		if (VoxelFormatB != PF_Unknown)
+		{
+			const FRHITextureCreateDesc Desc =
+				FRHITextureCreateDesc::Create3D(TEXT("SparseVolumeTexture.PhysicalTileDataB.RHITexture"),
+					TileDataVolumeResolution.X, TileDataVolumeResolution.Y, TileDataVolumeResolution.Z, VoxelFormatB)
+				.SetFlags(ETextureCreateFlags::ShaderResource);
+
+			PhysicalTileDataBTextureRHI = RHICreateTexture(Desc);
+
+			const int32 FormatSize = GPixelFormats[VoxelFormatB].BlockBytes;
+			RHIUpdateTexture3D(PhysicalTileDataBTextureRHI, 0, UpdateRegion, UpdateRegion.Width * FormatSize, UpdateRegion.Width * UpdateRegion.Height * FormatSize, (const uint8*)SparseVolumeTextureRuntime->PhysicalTileDataB.GetData());
+		}
 	}
 }
 
@@ -549,13 +593,18 @@ void UAnimatedSparseVolumeTexture::ConvertRawSourceDataToSparseVolumeTextureRunt
 		FSparseVolumeRawSource SparseVolumeRawSource;
 		SparseVolumeRawSource.Serialize(RawDataArchiveReader);
 
+		FOpenVDBToSVTConversionResult SVTResult;
+		SVTResult.Header = &Frame.SparseVolumeTextureRuntime.Header;
+		SVTResult.PageTable = &Frame.SparseVolumeTextureRuntime.PageTable;
+		SVTResult.PhysicalTileDataA = &Frame.SparseVolumeTextureRuntime.PhysicalTileDataA;
+		SVTResult.PhysicalTileDataB = &Frame.SparseVolumeTextureRuntime.PhysicalTileDataB;
+
 		// Then, cook the runtime data and serialize it out.
 		const bool bSuccess = ConvertOpenVDBToSparseVolumeTextureDelegate.IsBound() && ConvertOpenVDBToSparseVolumeTextureDelegate.Execute(
 			SparseVolumeRawSource.SourceAssetFile,
 			SparseVolumeRawSource.PackedDataA,
-			&Frame.SparseVolumeTextureRuntime.Header,
-			&Frame.SparseVolumeTextureRuntime.DensityPage,
-			&Frame.SparseVolumeTextureRuntime.DensityData,
+			SparseVolumeRawSource.PackedDataB,
+			&SVTResult,
 			false, FVector::Zero(), FVector::Zero());
 		ensure(bSuccess);
 
diff --git a/Engine/Source/Runtime/Renderer/Private/SparseVolumeTexture/SparseVolumeTextureViewerComponent.h b/Engine/Source/Runtime/Renderer/Private/SparseVolumeTexture/SparseVolumeTextureViewerComponent.h
--- a/Engine/Source/Runtime/Renderer/Private/SparseVolumeTexture/SparseVolumeTextureViewerComponent.h
+++ b/Engine/Source/Runtime/Renderer/Private/SparseVolumeTexture/SparseVolumeTextureViewerComponent.h
@@ -35,7 +35,7 @@ class USparseVolumeTextureViewerComponent : public UPrimitiveComponent
 	UPROPERTY(EditAnywhere, BlueprintReadWrite, Category = "Asset Preview", meta = (UIMin = 0.0, UIMax = 1.0, ClampMin = 0.0, ClampMax = 1.0, EditCondition = "!bAnimate"))
 	float AnimationFrame;
 
-	UPROPERTY(EditAnywhere, BlueprintReadWrite, Category = "Asset Preview", meta = (UIMin = 0, UIMax = 4, ClampMin = 0, ClampMax = 3))
+	UPROPERTY(EditAnywhere, BlueprintReadWrite, Category = "Asset Preview", meta = (UIMin = 0, UIMax = 7, ClampMin = 0, ClampMax = 7))
 	int32 ComponentToVisualize;
 
 	//~ Begin UPrimitiveComponent Interface.
diff --git a/Engine/Source/Runtime/Renderer/Private/SparseVolumeTexture/SparseVolumeTextureViewerRendering.cpp b/Engine/Source/Runtime/Renderer/Private/SparseVolumeTexture/SparseVolumeTextureViewerRendering.cpp
--- a/Engine/Source/Runtime/Renderer/Private/SparseVolumeTexture/SparseVolumeTextureViewerRendering.cpp
+++ b/Engine/Source/Runtime/Renderer/Private/SparseVolumeTexture/SparseVolumeTextureViewerRendering.cpp
@@ -64,7 +64,8 @@ class FVisualizeSparseVolumeTexturePS : public FGlobalShader
 		SHADER_PARAMETER(FVector3f, SparseVolumeTextureResolution)
 		SHADER_PARAMETER(FVector3f, SparseVolumeTexturePageTableResolution)
 		SHADER_PARAMETER(uint32, ComponentToVisualize)
-		SHADER_PARAMETER_TEXTURE(Texture3D, SparseVolumeTexture)
+		SHADER_PARAMETER_TEXTURE(Texture3D, SparseVolumeTextureA)
+		SHADER_PARAMETER_TEXTURE(Texture3D, SparseVolumeTextureB)
 		SHADER_PARAMETER_TEXTURE(Texture3D<uint>, SparseVolumeTexturePageTable)
 	END_SHADER_PARAMETER_STRUCT()
 
@@ -132,7 +133,8 @@ void AddSparseVolumeTextureViewerRenderPass(FRDGBuilder& GraphBuilder, FSceneRen
 			PsPassParameters->VolumeBoundMaxWorld = FVector3f(SVTProxyBound.Origin + SVTProxyBound.BoxExtent);
 			PsPassParameters->SparseVolumeTextureResolution = FVector3f::OneVector;
 			PsPassParameters->SparseVolumeTexturePageTableResolution = FVector3f::OneVector;
-			PsPassParameters->SparseVolumeTexture = GBlackVolumeTexture->TextureRHI;
+			PsPassParameters->SparseVolumeTextureA = GBlackVolumeTexture->TextureRHI;
+			PsPassParameters->SparseVolumeTextureB = GBlackVolumeTexture->TextureRHI;
 			PsPassParameters->SparseVolumeTexturePageTable = GBlackUintVolumeTexture->TextureRHI;
 
 
@@ -143,10 +145,15 @@ void AddSparseVolumeTextureViewerRenderPass(FRDGBuilder& GraphBuilder, FSceneRen
 				PsPassParameters->SparseVolumeTextureResolution = FVector3f(Header.SourceVolumeResolution);
 				PsPassParameters->SparseVolumeTexturePageTableResolution = FVector3f(Header.PageTableVolumeResolution);
 
-				FRHITexture* Texture = SVTProxy->SparseVolumeTextureSceneProxy->GetTileDataTextureRHI();
-				if (Texture)
+				FRHITexture* TextureA = SVTProxy->SparseVolumeTextureSceneProxy->GetPhysicalTileDataATextureRHI();
+				if (TextureA)
 				{
-					PsPassParameters->SparseVolumeTexture = Texture;
+					PsPassParameters->SparseVolumeTextureA = TextureA;
+				}
+				FRHITexture* TextureB = SVTProxy->SparseVolumeTextureSceneProxy->GetPhysicalTileDataBTextureRHI();
+				if (TextureB)
+				{
+					PsPassParameters->SparseVolumeTextureB = TextureB;
 				}
 				FRHITexture* PageTableTexture = SVTProxy->SparseVolumeTextureSceneProxy->GetPageTableTextureRHI();
 				if (PageTableTexture)
