Implemented sparse iteration when sampling the volume data from OpenVDB during import.

#rnx

diff --git a/Engine/Source/Editor/SparseVolumeTexture/Private/OpenVDBGridAdapter.h b/Engine/Source/Editor/SparseVolumeTexture/Private/OpenVDBGridAdapter.h
--- a/Engine/Source/Editor/SparseVolumeTexture/Private/OpenVDBGridAdapter.h
+++ b/Engine/Source/Editor/SparseVolumeTexture/Private/OpenVDBGridAdapter.h
@@ -19,6 +19,24 @@ using FOpenVDBDouble2Grid = openvdb::Grid<openvdb::tree::Tree4<openvdb::Vec2d, 5
 using FOpenVDBDouble3Grid = openvdb::Vec3DGrid;
 using FOpenVDBDouble4Grid = openvdb::Grid<openvdb::tree::Tree4<openvdb::Vec4d, 5, 4, 3>::Type>;
 
+template<typename ValueType>
+constexpr uint32 GetOpenVDBValueNumComponents()
+{
+	return (uint32)ValueType::size;
+}
+
+template<>
+constexpr uint32 GetOpenVDBValueNumComponents<float>()
+{
+	return 1;
+}
+
+template<>
+constexpr uint32 GetOpenVDBValueNumComponents<double>()
+{
+	return 1;
+}
+
 template<typename ValueType>
 float GetOpenVDBValueComponent(const ValueType& Value, uint32 ComponentIndex)
 {
@@ -77,7 +95,8 @@ double GetOpenVDBMinComponent<double>(const double& ValueA, const double& ValueB
 class IOpenVDBGridAdapterBase
 {
 public:
-	virtual float Sample(const openvdb::Coord& VoxelCoord, uint32 ComponentIndex) const = 0;
+	virtual void IteratePhysical(TFunctionRef<void(const FIntVector3& Coord, uint32 NumComponents, float* VoxelValues)> OnVisit) const = 0;
+	virtual float Sample(const FIntVector3& VoxelCoord, uint32 ComponentIndex) const = 0;
 	virtual void GetMinMaxValue(uint32 ComponentIndex, float* OutMinVal, float* OutMaxVal) const = 0;
 	virtual float GetBackgroundValue(uint32 ComponentIndex) const = 0;
 	virtual ~IOpenVDBGridAdapterBase() = default;
@@ -93,9 +112,31 @@ public:
 	{
 	}
 
-	float Sample(const openvdb::Coord& VoxelCoord, uint32 ComponentIndex) const override
+	void IteratePhysical(TFunctionRef<void(const FIntVector3& Coord, uint32 NumComponents, float* VoxelValues)> OnVisit) const override
+	{
+		constexpr uint32 NumComponents = GetOpenVDBValueNumComponents<ValueType>();
+		static_assert(NumComponents <= 4);
+		for (auto LeafIt = Grid->constTree().cbeginLeaf(); LeafIt; ++LeafIt)
+		{
+			for (auto ValueIt = LeafIt->cbeginValueOn(); ValueIt; ++ValueIt)
+			{
+				const ValueType VoxelValue = ValueIt.getValue();
+				float VoxelValueComponents[4]{};
+				for (uint32 ComponentIdx = 0; ComponentIdx < NumComponents; ++ComponentIdx)
+				{
+					VoxelValueComponents[ComponentIdx] = GetOpenVDBValueComponent(VoxelValue, ComponentIdx);
+				}
+
+				const openvdb::Coord CoordVDB = ValueIt.getCoord();
+				const FIntVector3 Coord(CoordVDB[0], CoordVDB[1], CoordVDB[2]);
+				OnVisit(Coord, NumComponents, VoxelValueComponents);
+			}
+		}
+	}
+
+	float Sample(const FIntVector3& VoxelCoord, uint32 ComponentIndex) const override
 	{
-		ValueType VoxelValue = Accessor.getValue(VoxelCoord);
+		ValueType VoxelValue = Accessor.getValue(openvdb::Coord(VoxelCoord.X, VoxelCoord.Y, VoxelCoord.Z));
 		return GetOpenVDBValueComponent(VoxelValue, ComponentIndex);
 	}
 
diff --git a/Engine/Source/Editor/SparseVolumeTexture/Private/SparseVolumeTextureOpenVDBUtility.cpp b/Engine/Source/Editor/SparseVolumeTexture/Private/SparseVolumeTextureOpenVDBUtility.cpp
--- a/Engine/Source/Editor/SparseVolumeTexture/Private/SparseVolumeTextureOpenVDBUtility.cpp
+++ b/Engine/Source/Editor/SparseVolumeTexture/Private/SparseVolumeTextureOpenVDBUtility.cpp
@@ -11,6 +11,10 @@ DEFINE_LOG_CATEGORY_STATIC(LogSparseVolumeTextureOpenVDBUtility, Log, All);
 
 #if OPENVDB_AVAILABLE
 
+// We are using this instead of GMaxVolumeTextureDimensions to be independent of the platform that
+// the asset is imported on. 2048 should be a safe value that should be supported by all our platforms.
+static constexpr int32 SVTMaxVolumeTextureDim = 2048;
+
 namespace
 {
 	// Utility class acting as adapter between TArray<uint8> and std::istream
@@ -255,6 +259,24 @@ static EPixelFormat GetMultiComponentFormat(ESparseVolumePackedDataFormat Format
 	return PF_Unknown;
 }
 
+static uint32 PackPageTableEntry(const FIntVector3& Coord)
+{
+	// A page encodes the physical tile coord as unsigned int of 11 11 10 bits
+	// This means a page coord cannot be larger than 2047 for x and y and 1023 for z
+	// which mean we cannot have more than 2048*2048*1024 = 4 Giga tiles of 16^3 tiles.
+	uint32 Result = (Coord.X & 0x7FF) | ((Coord.Y & 0x7FF) << 11) | ((Coord.Z & 0x3FF) << 22);
+	return Result;
+}
+
+static FIntVector3 UnpackPageTableEntry(uint32 Packed)
+{
+	FIntVector3 Result;
+	Result.X = Packed & 0x7FF;
+	Result.Y = (Packed >> 11) & 0x7FF;
+	Result.Z = (Packed >> 22) & 0x3FF;
+	return Result;
+}
+
 bool ConvertOpenVDBToSparseVolumeTexture(
 	TArray<uint8>& SourceFile,
 	FSparseVolumeRawSourcePackedData& PackedDataA,
@@ -277,7 +299,7 @@ bool ConvertOpenVDBToSparseVolumeTexture(
 	bool bNormalizedFormat[NumPackedData] = {};
 	bool bHasValidSourceGrids[NumPackedData] = {};
 	bool bAnySourceGridIndicesValid = false;
-	
+
 	for (uint32 PackedDataIdx = 0; PackedDataIdx < NumPackedData; ++PackedDataIdx)
 	{
 		uint32 NumRequiredComponents = 0;
@@ -300,7 +322,7 @@ bool ConvertOpenVDBToSparseVolumeTexture(
 
 			if (MultiCompFormat[PackedDataIdx] == PF_Unknown)
 			{
-				// SVT_TODO error message about unsupported format
+				UE_LOG(LogSparseVolumeTextureOpenVDBUtility, Warning, TEXT("SparseVolumeTexture is set to use an unsupported format: %i"), (int32)PackedData[PackedDataIdx]->Format);
 				return false;
 			}
 
@@ -312,7 +334,7 @@ bool ConvertOpenVDBToSparseVolumeTexture(
 	// All source grid indices are INDEX_NONE, so nothing was selected for import
 	if (!bAnySourceGridIndicesValid)
 	{
-		// SVT_TODO error message
+		UE_LOG(LogSparseVolumeTextureOpenVDBUtility, Warning, TEXT("SparseVolumeTexture has all components set to <None>, so there is nothing to import."));
 		return false;
 	}
 
@@ -331,7 +353,7 @@ bool ConvertOpenVDBToSparseVolumeTexture(
 			const uint32 SourceGridIndex = PackedData[PackedDataIdx]->SourceGridIndex[CompIdx];
 			if (SourceGridIndex != INDEX_NONE && SourceGridIndex >= NumSourceGrids)
 			{
-				// SVT_TODO error message
+				UE_LOG(LogSparseVolumeTextureOpenVDBUtility, Warning, TEXT("SparseVolumeTexture has invalid index into the array of grids in the source file: %i"), (int32)SourceGridIndex);
 				return false;
 			}
 		}
@@ -346,7 +368,7 @@ bool ConvertOpenVDBToSparseVolumeTexture(
 
 	// Compute per source grid data of up to 4 different grids (one per component)
 	TArray<TSharedPtr<IOpenVDBGridAdapterBase>> UniqueGridAdapters;
-	UniqueGridAdapters.SetNum((int32)Grids->size());
+	UniqueGridAdapters.SetNum((int32)NumSourceGrids);
 	TSharedPtr<IOpenVDBGridAdapterBase> GridAdapters[NumPackedData][4]{};
 	float GridBackgroundValues[NumPackedData][4]{};
 	float NormalizeScale[NumPackedData][4]{};
@@ -388,7 +410,7 @@ bool ConvertOpenVDBToSparseVolumeTexture(
 			{
 				OVDBData.VolumeActiveAABBMin = ActiveMin;
 				OVDBData.VolumeActiveAABBMax = ActiveMax;
-				OVDBData.VolumeActiveDim = ActiveMax - ActiveMin;
+				OVDBData.VolumeActiveDim = ActiveMax + 1 - ActiveMin;
 			}
 
 			Header.SourceVolumeResolution.X = FMath::Max(Header.SourceVolumeResolution.X, OVDBData.VolumeActiveDim.X);
@@ -411,14 +433,62 @@ bool ConvertOpenVDBToSparseVolumeTexture(
 		}
 	}
 	
-	// SVT_TODO handle error too high volume resolution assuming dx11 limitations?
-	Header.PageTableVolumeResolution = FIntVector3(
+	FIntVector3 PageTableVolumeResolution = FIntVector3(
 		FMath::DivideAndRoundUp(Header.SourceVolumeResolution.X, SPARSE_VOLUME_TILE_RES),
 		FMath::DivideAndRoundUp(Header.SourceVolumeResolution.Y, SPARSE_VOLUME_TILE_RES),
 		FMath::DivideAndRoundUp(Header.SourceVolumeResolution.Z, SPARSE_VOLUME_TILE_RES));
+	if (PageTableVolumeResolution.X > SVTMaxVolumeTextureDim
+		|| PageTableVolumeResolution.Y > SVTMaxVolumeTextureDim
+		|| PageTableVolumeResolution.Z > SVTMaxVolumeTextureDim)
+	{
+		UE_LOG(LogSparseVolumeTextureOpenVDBUtility, Warning, TEXT("SparseVolumeTexture page table texture dimensions exceed limit (%ix%ix%i): %ix%ix%i"), SVTMaxVolumeTextureDim, SVTMaxVolumeTextureDim, SVTMaxVolumeTextureDim, PageTableVolumeResolution.X, PageTableVolumeResolution.Y, PageTableVolumeResolution.Z);
+		return false;
+	}
+
+	Header.PageTableVolumeResolution = PageTableVolumeResolution;
 	Header.TileDataVolumeResolution = FIntVector::ZeroValue;	// unknown for now
 
-	uint32 NumAllocatedPages = 0;
+	// Tag all pages with valid data
+	TBitArray PagesWithData(false, Header.PageTableVolumeResolution.Z * Header.PageTableVolumeResolution.Y * Header.PageTableVolumeResolution.X);
+	for (uint32 GridIdx = 0; GridIdx < NumSourceGrids; ++GridIdx)
+	{
+		if (!UniqueGridAdapters[GridIdx])
+		{
+			continue;
+		}
+
+		UniqueGridAdapters[GridIdx]->IteratePhysical(
+			[&](const FIntVector3& Coord, uint32 NumComponents, float* VoxelValues)
+			{
+				// Check if the source grid component is used at all
+				bool bValueIsUsed = false;
+				for (uint32 PackedDataIdx = 0; PackedDataIdx < NumPackedData && !bValueIsUsed; ++PackedDataIdx)
+				{
+					for (uint32 DstCompIdx = 0; DstCompIdx < NumActualComponents[PackedDataIdx] && !bValueIsUsed; ++DstCompIdx)
+					{
+						for (uint32 SrcCompIdx = 0; SrcCompIdx < NumComponents && !bValueIsUsed; ++SrcCompIdx)
+						{
+							const bool bIsBackgroundValue = (VoxelValues[SrcCompIdx] == GridBackgroundValues[PackedDataIdx][DstCompIdx]);
+							bValueIsUsed |= !bIsBackgroundValue && (PackedData[PackedDataIdx]->SourceGridIndex[DstCompIdx] == GridIdx) && (PackedData[PackedDataIdx]->SourceComponentIndex[DstCompIdx] == SrcCompIdx);
+						}
+					}
+				}
+				if (!bValueIsUsed)
+				{
+					return;
+				}
+
+				const FIntVector3 GridCoord = Coord - SmallestAABBMin;
+				check(GridCoord.X >= 0 && GridCoord.Y >= 0 && GridCoord.Z >= 0);
+				check(GridCoord.X < Header.SourceVolumeResolution.X && GridCoord.Y < Header.SourceVolumeResolution.Y && GridCoord.Z < Header.SourceVolumeResolution.Z);
+				const FIntVector3 PageCoord = GridCoord / SPARSE_VOLUME_TILE_RES;
+				check(PageCoord.X >= 0 && PageCoord.Y >= 0 && PageCoord.Z >= 0);
+				check(PageCoord.X < Header.PageTableVolumeResolution.X && PageCoord.Y < Header.PageTableVolumeResolution.Y && PageCoord.Z < Header.PageTableVolumeResolution.Z);
+
+				const int32 PageIndex = PageCoord.Z * (Header.PageTableVolumeResolution.Y * Header.PageTableVolumeResolution.X) + PageCoord.Y * Header.PageTableVolumeResolution.X + PageCoord.X;
+				PagesWithData[PageIndex] = true;
+			});
+	}
 
 	// Allocate some memory for temp data (worst case)
 	TArray<FIntVector3> LinearAllocatedPages;
@@ -427,40 +497,15 @@ bool ConvertOpenVDBToSparseVolumeTexture(
 	// Go over each potential page from the source data and push allocate it if it has any data.
 	// Otherwise point to the default empty page.
 	bool bAnyEmptyPageExists = false;
+	uint32 NumAllocatedPages = 0;
 	for (int32_t PageZ = 0; PageZ < Header.PageTableVolumeResolution.Z; ++PageZ)
 	{
 		for (int32_t PageY = 0; PageY < Header.PageTableVolumeResolution.Y; ++PageY)
 		{
 			for (int32_t PageX = 0; PageX < Header.PageTableVolumeResolution.X; ++PageX)
 			{
-				bool bHasAnyData = false;
-
-				for (uint32 PackedDataIdx = 0; PackedDataIdx < NumPackedData; ++PackedDataIdx)
-				{
-					if (!bHasValidSourceGrids[PackedDataIdx])
-					{
-						continue;
-					}
-					for (int32_t z = 0; z < SPARSE_VOLUME_TILE_RES; ++z)
-					{
-						for (int32_t y = 0; y < SPARSE_VOLUME_TILE_RES; ++y)
-						{
-							for (int32_t x = 0; x < SPARSE_VOLUME_TILE_RES; ++x)
-							{
-								for (uint32 CompIdx = 0; CompIdx < NumActualComponents[PackedDataIdx]; ++CompIdx)
-								{
-									if (PackedData[PackedDataIdx]->SourceGridIndex[CompIdx] != INDEX_NONE)
-									{
-										FVector VoxelCoord = FVector(SmallestAABBMin) + FVector(PageX, PageY, PageZ) * SPARSE_VOLUME_TILE_RES + FVector(x, y, z);	// This assumes sampling outside the boundary returns a default value
-										float VoxelValue = GridAdapters[PackedDataIdx][CompIdx]->Sample(openvdb::Coord(VoxelCoord.X, VoxelCoord.Y, VoxelCoord.Z), PackedData[PackedDataIdx]->SourceComponentIndex[CompIdx]);
-										bHasAnyData |= VoxelValue != GridBackgroundValues[PackedDataIdx][CompIdx];
-									}
-								}
-							}
-						}
-					}
-				}
-
+				const int32 PageIndex = PageZ * (Header.PageTableVolumeResolution.Y * Header.PageTableVolumeResolution.X) + PageY * Header.PageTableVolumeResolution.X + PageX;
+				bool bHasAnyData = PagesWithData[PageIndex];
 				if (bHasAnyData)
 				{
 					LinearAllocatedPages[NumAllocatedPages] = FIntVector3(PageX, PageY, PageZ);
@@ -485,6 +530,13 @@ bool ConvertOpenVDBToSparseVolumeTexture(
 	}
 	const FIntVector3 TileCoordResolution = Header.TileDataVolumeResolution;
 	Header.TileDataVolumeResolution = Header.TileDataVolumeResolution * SPARSE_VOLUME_TILE_RES;
+	if (Header.TileDataVolumeResolution.X > SVTMaxVolumeTextureDim
+		|| Header.TileDataVolumeResolution.Y > SVTMaxVolumeTextureDim
+		|| Header.TileDataVolumeResolution.Z > SVTMaxVolumeTextureDim)
+	{
+		UE_LOG(LogSparseVolumeTextureOpenVDBUtility, Warning, TEXT("SparseVolumeTexture tile data texture dimensions exceed limit (%ix%ix%i): %ix%ix%i"), SVTMaxVolumeTextureDim, SVTMaxVolumeTextureDim, SVTMaxVolumeTextureDim, Header.TileDataVolumeResolution.X, Header.TileDataVolumeResolution.Y, Header.TileDataVolumeResolution.Z);
+		return false;
+	}
 
 	// Initialise the SparseVolumeTexture page and tile.
 	OutResult->PageTable->SetNumZeroed(Header.PageTableVolumeResolution.X * Header.PageTableVolumeResolution.Y * Header.PageTableVolumeResolution.Z);
@@ -493,92 +545,112 @@ bool ConvertOpenVDBToSparseVolumeTexture(
 	uint32* PageTablePtr = OutResult->PageTable->GetData();
 	uint8* PhysicalTileDataPtrs[] = { OutResult->PhysicalTileDataA->GetData(), OutResult->PhysicalTileDataB->GetData() };
 
-	FIntVector DestinationTileCoord = FIntVector::ZeroValue;
-	auto GoToNextTileCoord = [&]()
+	// Generate page table and tile volume data by splatting the data
 	{
-		DestinationTileCoord.X++;
-		if (DestinationTileCoord.X >= TileCoordResolution.X)
+		FIntVector3 DstTileCoord = FIntVector3::ZeroValue;
+		auto GoToNextTileCoord = [&]()
 		{
-			DestinationTileCoord.X = 0;
-			DestinationTileCoord.Y++;
-		}
-		if (DestinationTileCoord.Y >= TileCoordResolution.Y)
+			DstTileCoord.X++;
+			if (DstTileCoord.X >= TileCoordResolution.X)
+			{
+				DstTileCoord.X = 0;
+				DstTileCoord.Y++;
+			}
+			if (DstTileCoord.Y >= TileCoordResolution.Y)
+			{
+				DstTileCoord.Y = 0;
+				DstTileCoord.Z++;
+			}
+		};
+
+		// Add an empty tile is needed, reserve slot at coord 0
+		if (bAnyEmptyPageExists)
 		{
-			DestinationTileCoord.Y = 0;
-			DestinationTileCoord.Z++;
+			// PageTable is all cleared to zero, simply skip a tile
+			GoToNextTileCoord();
 		}
-	};
 
-	// Add an empty tile is needed, reserve slot at coord 0
-	if (bAnyEmptyPageExists)
-	{
-		// PageTable is all cleared to zero, simply skip a tile
-		GoToNextTileCoord();
+		for (uint32 i = 0; i < NumAllocatedPages; ++i)
+		{
+			FIntVector3 PageCoordToSplat = LinearAllocatedPages[i];
+			uint32 DestinationTileCoord32bit = PackPageTableEntry(DstTileCoord);
+
+			// Setup the page table entry
+			PageTablePtr
+				[
+					PageCoordToSplat.Z * Header.PageTableVolumeResolution.X * Header.PageTableVolumeResolution.Y +
+					PageCoordToSplat.Y * Header.PageTableVolumeResolution.X +
+				PageCoordToSplat.X
+				] = DestinationTileCoord32bit;
+
+			// Set the next tile to be written to
+			GoToNextTileCoord();
+		}
 	}
 
-	// Generate page table and tile volume data by splatting the data
-	for (uint32 i = 0; i < NumAllocatedPages; ++i)
+	for (uint32 GridIdx = 0; GridIdx < NumSourceGrids; ++GridIdx)
 	{
-		FIntVector3 PageCoordToSplat = LinearAllocatedPages[i];
-
-		// A page encodes the physical tile coord as unsigned int of 11 11 10 bits
-		// This means a page coord cannot be larger than 2047 for x and y and 1023 for z
-		// which mean we cannot have more than 2048*2048*1024 = 4 Giga tiles of 16^3 tiles.
-		uint32 DestinationTileCoord32bit = (DestinationTileCoord.X & 0x7FF) | ((DestinationTileCoord.Y & 0x7FF) << 11) | ((DestinationTileCoord.Z & 0x3FF) << 22);
-
-		// Setup the page table entry
-		PageTablePtr
-			[
-				PageCoordToSplat.Z * Header.PageTableVolumeResolution.X * Header.PageTableVolumeResolution.Y +
-				PageCoordToSplat.Y * Header.PageTableVolumeResolution.X +
-				PageCoordToSplat.X
-			] = DestinationTileCoord32bit;
-
-		// Now copy the tile data from the source to the page atlas
-		for (uint32 PackedDataIdx = 0; PackedDataIdx < NumPackedData; ++PackedDataIdx)
+		if (!UniqueGridAdapters[GridIdx])
 		{
-			if (!bHasValidSourceGrids[PackedDataIdx])
-			{
-				continue;
-			}
-			uint8* PhysicalTileDataPtr = PhysicalTileDataPtrs[PackedDataIdx];
-			for (int32_t z = 0; z < SPARSE_VOLUME_TILE_RES; ++z)
+			continue;
+		}
+
+		UniqueGridAdapters[GridIdx]->IteratePhysical(
+			[&](const FIntVector3& Coord, uint32 NumComponents, float* VoxelValues)
 			{
-				for (int32_t y = 0; y < SPARSE_VOLUME_TILE_RES; ++y)
+				const FIntVector3 GridCoord = Coord - SmallestAABBMin;
+				check(GridCoord.X >= 0 && GridCoord.Y >= 0 && GridCoord.Z >= 0);
+				check(GridCoord.X < Header.SourceVolumeResolution.X&& GridCoord.Y < Header.SourceVolumeResolution.Y&& GridCoord.Z < Header.SourceVolumeResolution.Z);
+				const FIntVector3 PageCoord = GridCoord / SPARSE_VOLUME_TILE_RES;
+				check(PageCoord.X >= 0 && PageCoord.Y >= 0 && PageCoord.Z >= 0);
+				check(PageCoord.X < Header.PageTableVolumeResolution.X&& PageCoord.Y < Header.PageTableVolumeResolution.Y&& PageCoord.Z < Header.PageTableVolumeResolution.Z);
+
+				const int32 PageIndex = PageCoord.Z * (Header.PageTableVolumeResolution.Y * Header.PageTableVolumeResolution.X) + PageCoord.Y * Header.PageTableVolumeResolution.X + PageCoord.X;
+
+				if (!PagesWithData[PageIndex])
 				{
-					for (int32_t x = 0; x < SPARSE_VOLUME_TILE_RES; ++x)
+					return;
+				}
+
+				const FIntVector3 DstTileCoord = UnpackPageTableEntry(PageTablePtr[PageIndex]);
+				const FIntVector3 TileLocalCoord = GridCoord % SPARSE_VOLUME_TILE_RES;
+
+				// Check all output components and splat VoxelValue if they map to this source grid/component
+				for (uint32 PackedDataIdx = 0; PackedDataIdx < NumPackedData; ++PackedDataIdx)
+				{
+					for (uint32 DstCompIdx = 0; DstCompIdx < NumActualComponents[PackedDataIdx]; ++DstCompIdx)
 					{
-						for (uint32 CompIdx = 0; CompIdx < NumActualComponents[PackedDataIdx]; ++CompIdx)
+						for (uint32 SrcCompIdx = 0; SrcCompIdx < NumComponents; ++SrcCompIdx)
 						{
-							float VoxelValue = 0.0f;
-							float VoxelValueNormalized = 0.0f;
-							if (PackedData[PackedDataIdx]->SourceGridIndex[CompIdx] != INDEX_NONE)
+							if ((PackedData[PackedDataIdx]->SourceGridIndex[DstCompIdx] != GridIdx) || (PackedData[PackedDataIdx]->SourceComponentIndex[DstCompIdx] != SrcCompIdx))
 							{
-								FVector VoxelCoord = FVector(SmallestAABBMin) + FVector(PageCoordToSplat * SPARSE_VOLUME_TILE_RES) + FVector(x, y, z);	// This assumes sampling outside the boundary returns a default value
-								VoxelValue = GridAdapters[PackedDataIdx][CompIdx]->Sample(openvdb::Coord(VoxelCoord.X, VoxelCoord.Y, VoxelCoord.Z), PackedData[PackedDataIdx]->SourceComponentIndex[CompIdx]);
-								VoxelValueNormalized = FMath::Clamp(VoxelValue * NormalizeScale[PackedDataIdx][CompIdx] + NormalizeBias[PackedDataIdx][CompIdx], 0.0f, 1.0f);
+								continue;
 							}
 
-							const SIZE_T DstCoord = ((DestinationTileCoord.Z * SPARSE_VOLUME_TILE_RES + z) * Header.TileDataVolumeResolution.X * Header.TileDataVolumeResolution.Y +
-								(DestinationTileCoord.Y * SPARSE_VOLUME_TILE_RES + y) * Header.TileDataVolumeResolution.X +
-								(DestinationTileCoord.X * SPARSE_VOLUME_TILE_RES + x)) * FormatSize[PackedDataIdx] + CompIdx * SingleComponentFormatSize[PackedDataIdx];
+							const float VoxelValueNormalized = FMath::Clamp(VoxelValues[SrcCompIdx] * NormalizeScale[PackedDataIdx][DstCompIdx] + NormalizeBias[PackedDataIdx][DstCompIdx], 0.0f, 1.0f);
+
+							const size_t DstVoxelIndex =
+								(DstTileCoord.Z * SPARSE_VOLUME_TILE_RES + TileLocalCoord.Z) * Header.TileDataVolumeResolution.X * Header.TileDataVolumeResolution.Y +
+								(DstTileCoord.Y * SPARSE_VOLUME_TILE_RES + TileLocalCoord.Y) * Header.TileDataVolumeResolution.X +
+								(DstTileCoord.X * SPARSE_VOLUME_TILE_RES + TileLocalCoord.X);
+							const size_t DstCoord = DstVoxelIndex * FormatSize[PackedDataIdx] + DstCompIdx * SingleComponentFormatSize[PackedDataIdx];
 
 							switch (PackedData[PackedDataIdx]->Format)
 							{
 							case ESparseVolumePackedDataFormat::Unorm8:
 							{
-								PhysicalTileDataPtr[DstCoord] = uint8(VoxelValueNormalized * 255.0f);
+								PhysicalTileDataPtrs[PackedDataIdx][DstCoord] = uint8(VoxelValueNormalized * 255.0f);
 								break;
 							}
 							case ESparseVolumePackedDataFormat::Float16:
 							{
-								const uint16 VoxelValue16FEncoded = FFloat16(VoxelValue).Encoded;
-								*((uint16*)(&PhysicalTileDataPtr[DstCoord])) = VoxelValue16FEncoded;
+								const uint16 VoxelValue16FEncoded = FFloat16(VoxelValues[SrcCompIdx]).Encoded;
+								*((uint16*)(&PhysicalTileDataPtrs[PackedDataIdx][DstCoord])) = VoxelValue16FEncoded;
 								break;
 							}
 							case ESparseVolumePackedDataFormat::Float32:
 							{
-								*((float*)(&PhysicalTileDataPtr[DstCoord])) = VoxelValue;
+								*((float*)(&PhysicalTileDataPtrs[PackedDataIdx][DstCoord])) = VoxelValues[SrcCompIdx];
 								break;
 							}
 							default:
@@ -588,11 +660,7 @@ bool ConvertOpenVDBToSparseVolumeTexture(
 						}
 					}
 				}
-			}
-		}
-
-		// Set the next tile to be written to
-		GoToNextTileCoord();
+			});
 	}
 
 	return true;
