OpenVDB import dialog now uses a proper table/list view for displaying grid info and sets the "Import as Sequence" checkbox automatically. Improved the algorithm for finding files belonging to a sequence. Also did some slight refactoring, moving out some logic of the SOpenVDBImportWindow class into the factory.

#rnx

diff --git a/Engine/Source/Editor/SparseVolumeTexture/Private/OpenVDBImportWindow.cpp b/Engine/Source/Editor/SparseVolumeTexture/Private/OpenVDBImportWindow.cpp
--- a/Engine/Source/Editor/SparseVolumeTexture/Private/OpenVDBImportWindow.cpp
+++ b/Engine/Source/Editor/SparseVolumeTexture/Private/OpenVDBImportWindow.cpp
@@ -17,6 +17,11 @@
 
 #define LOCTEXT_NAMESPACE "SOpenVDBImportWindow"
 
+#define VDB_GRID_ROW_NAME_GRID_INDEX TEXT("GridIndex")
+#define VDB_GRID_ROW_NAME_GRID_TYPE TEXT("GridType")
+#define VDB_GRID_ROW_NAME_GRID_NAME TEXT("GridName")
+#define VDB_GRID_ROW_NAME_GRID_DIMS TEXT("GridDims")
+
 static FText GetGridComboBoxItemText(TSharedPtr<FOpenVDBGridComponentInfo> InItem)
 {
 	return InItem ? FText::FromString(InItem->DisplayString) : LOCTEXT("NoneGrid", "<None>");
@@ -41,6 +46,10 @@ void SOpenVDBImportWindow::Construct(const FArguments& InArgs)
 {
 	PackedDataA = InArgs._PackedDataA;
 	PackedDataB = InArgs._PackedDataB;
+	DefaultAssignmentA = *PackedDataA;
+	DefaultAssignmentB = *PackedDataB;
+	bIsSequence = InArgs._NumFoundFiles > 1;
+	OpenVDBGridInfo = InArgs._OpenVDBGridInfo;
 	OpenVDBGridComponentInfo = InArgs._OpenVDBGridComponentInfo;
 	OpenVDBSupportedTargetFormats = InArgs._OpenVDBSupportedTargetFormats;
 	WidgetWindow = InArgs._WidgetWindow;
@@ -144,7 +153,16 @@ void SOpenVDBImportWindow::Construct(const FArguments& InArgs)
 			[
 				SAssignNew(ImportAsSequenceCheckBox, SCheckBox)
 				.ToolTipText(LOCTEXT("ImportAsSequenceCheckBoxTooltip", "Import multiple sequentially labeled .vdb files as a single animated SparseVirtualTexture sequence."))
-				.IsChecked(false)
+				.IsChecked(bIsSequence)
+			]
+			+ SHorizontalBox::Slot()
+			.VAlign(VAlign_Center)
+			.HAlign(HAlign_Right)
+			.AutoWidth()
+			.Padding(2.0f)
+			[
+				SNew(STextBlock)
+				.Text(FText::FromString(FString::Format(TEXT("Found {0} File(s)"), { InArgs._NumFoundFiles })))
 			]
 		]
 		+ SVerticalBox::Slot()
@@ -182,8 +200,19 @@ void SOpenVDBImportWindow::Construct(const FArguments& InArgs)
 			.Padding(FMargin(3))
 			.BorderImage(FAppStyle::GetBrush("ToolPanel.GroupBorder"))
 			[
-				SNew(STextBlock)
-				.Text(FText::FromString(InArgs._FileInfoString))
+				SNew(SListView<TSharedPtr<FOpenVDBGridInfo>>)
+				.ItemHeight(24)
+				.ScrollbarVisibility(EVisibility::Visible)
+				.ListItemsSource(OpenVDBGridInfo)
+				.OnGenerateRow(this, &SOpenVDBImportWindow::GenerateGridInfoItemRow)
+				.HeaderRow
+				(
+					SNew(SHeaderRow)
+					+ SHeaderRow::Column(VDB_GRID_ROW_NAME_GRID_INDEX).DefaultLabel(LOCTEXT("GridIndex", "Index")).FillWidth(0.05f)
+					+ SHeaderRow::Column(VDB_GRID_ROW_NAME_GRID_NAME).DefaultLabel(LOCTEXT("GridName", "Name")).FillWidth(0.15f)
+					+ SHeaderRow::Column(VDB_GRID_ROW_NAME_GRID_TYPE).DefaultLabel(LOCTEXT("GridType", "Type")).FillWidth(0.1f)
+					+ SHeaderRow::Column(VDB_GRID_ROW_NAME_GRID_DIMS).DefaultLabel(LOCTEXT("GridDims", "Dimensions")).FillWidth(0.25f)
+				)
 			]
 		]
 	);
@@ -268,6 +297,11 @@ EActiveTimerReturnType SOpenVDBImportWindow::SetFocusPostConstruct(double InCurr
 	return EActiveTimerReturnType::Stop;
 }
 
+TSharedRef<ITableRow> SOpenVDBImportWindow::GenerateGridInfoItemRow(TSharedPtr<FOpenVDBGridInfo> Item, const TSharedRef<STableViewBase>& OwnerTable)
+{
+	return SNew(SOpenVDBGridInfoTableRow, OwnerTable).OpenVDBGridInfo(Item);
+}
+
 bool SOpenVDBImportWindow::CanImport() const
 {
 	const FUintVector4& GridIndicesA = PackedDataA->SourceGridIndex;
@@ -302,42 +336,10 @@ void SOpenVDBImportWindow::SetDefaultGridAssignment()
 {
 	check(OpenVDBGridComponentInfo);
 
-	FSparseVolumeRawSourcePackedData* PackedData[] = { PackedDataA , PackedDataB };
-	for (FSparseVolumeRawSourcePackedData* Data : PackedData)
-	{
-		Data->Format = ESparseVolumePackedDataFormat::Float32;
-		Data->SourceGridIndex = FUintVector4(INDEX_NONE);
-		Data->SourceComponentIndex = FUintVector4(INDEX_NONE);
-		Data->bRemapInputForUnorm = false;
-	}
-
-	const TCHAR* SearchNames[] = { TEXT("density"), TEXT("heat"), TEXT("temperature"), TEXT("motion") };
-
-	// Try to find grids named like the SearchNames in the source file and assign them to the components of the output SVT.
-	uint32 CurrentOutputPackedData = 0;
-	uint32 CurrentOutputComponent = 0;
-	for (const TCHAR* SearchName : SearchNames)
-	{
-		for (const TSharedPtr<FOpenVDBGridComponentInfo>& GridComponent : *OpenVDBGridComponentInfo)
-		{
-			for (uint32 InputComponent = 0; InputComponent < 4; ++InputComponent)
-			{
-				if (GridComponent->Name == SearchName && GridComponent->ComponentIndex == InputComponent)
-				{
-					PackedData[CurrentOutputPackedData]->SourceGridIndex[CurrentOutputComponent] = GridComponent->Index;
-					PackedData[CurrentOutputPackedData]->SourceComponentIndex[CurrentOutputComponent] = GridComponent->ComponentIndex;
-					++CurrentOutputComponent;
-					if (CurrentOutputComponent == 4)
-					{
-						CurrentOutputComponent = 0;
-						++CurrentOutputPackedData;
-					}
-					break;
-				}
-			}
-		}
-	}
+	*PackedDataA = DefaultAssignmentA;
+	*PackedDataB = DefaultAssignmentB;
 
+	ImportAsSequenceCheckBox->SetIsChecked(bIsSequence ? ECheckBoxState::Checked : ECheckBoxState::Unchecked);
 	PackedDataAConfigurator->RefreshUIFromData();
 	PackedDataBConfigurator->RefreshUIFromData();
 }
@@ -565,4 +567,43 @@ void SOpenVDBPackedDataConfigurator::RefreshUIFromData()
 	RemapUnormCheckBox->SetIsChecked(PackedData->bRemapInputForUnorm);
 }
 
-#undef LOCTEXT_NAMESPACE
+void SOpenVDBGridInfoTableRow::Construct(const FArguments& InArgs, const TSharedRef<STableViewBase>& OwnerTableView)
+{
+	OpenVDBGridInfo = InArgs._OpenVDBGridInfo;
+	SMultiColumnTableRow<TSharedPtr<FOpenVDBGridInfo>>::Construct(FSuperRowType::FArguments(), OwnerTableView);
+}
+
+TSharedRef<SWidget> SOpenVDBGridInfoTableRow::GenerateWidgetForColumn(const FName& ColumnName)
+{
+	if (ColumnName == VDB_GRID_ROW_NAME_GRID_INDEX)
+	{
+		return SNew(SBox).Padding(2).VAlign(VAlign_Center)
+			[
+				SNew(STextBlock).Text(FText::FromString(FString::Format(TEXT("{0}."), { OpenVDBGridInfo->Index })))
+			];
+	}
+	else if (ColumnName == VDB_GRID_ROW_NAME_GRID_TYPE)
+	{
+		return SNew(SBox).Padding(2).VAlign(VAlign_Center)
+			[
+				SNew(STextBlock).Text(FText::FromString(OpenVDBGridTypeToString(OpenVDBGridInfo->Type)))
+			];
+	}
+	else if (ColumnName == VDB_GRID_ROW_NAME_GRID_NAME)
+	{
+		return SNew(SBox).Padding(2).VAlign(VAlign_Center)
+			[
+				SNew(STextBlock).Text(FText::FromString(OpenVDBGridInfo->Name))
+			];
+	}
+	else if (ColumnName == VDB_GRID_ROW_NAME_GRID_DIMS)
+	{
+		return SNew(SBox).Padding(2).VAlign(VAlign_Center)
+			[
+				SNew(STextBlock).Text(FText::FromString(OpenVDBGridInfo->VolumeActiveDim.ToString()))
+			];
+	}
+	return SNullWidget::NullWidget;
+}
+
+#undef LOCTEXT_NAMESPACE
\ No newline at end of file
diff --git a/Engine/Source/Editor/SparseVolumeTexture/Private/OpenVDBImportWindow.h b/Engine/Source/Editor/SparseVolumeTexture/Private/OpenVDBImportWindow.h
--- a/Engine/Source/Editor/SparseVolumeTexture/Private/OpenVDBImportWindow.h
+++ b/Engine/Source/Editor/SparseVolumeTexture/Private/OpenVDBImportWindow.h
@@ -7,7 +7,9 @@
 #include "Widgets/Input/SComboBox.h"
 #include "Widgets/Input/SCheckBox.h"
 
-struct ENGINE_API FSparseVolumeRawSourcePackedData;
+#include "SparseVolumeTexture/SparseVolumeTexture.h"
+
+struct FOpenVDBGridInfo;
 enum class ESparseVolumePackedDataFormat : uint8;
 
 struct FOpenVDBGridComponentInfo
@@ -18,18 +20,28 @@ struct FOpenVDBGridComponentInfo
 	FString DisplayString; // Contains source file grid index, name and component (if it is a multi component type like Float3)
 };
 
-class SOpenVDBComponentPicker : public SCompoundWidget
+class SOpenVDBGridInfoTableRow : public SMultiColumnTableRow<TSharedPtr<FOpenVDBGridInfo>>
 {
 public:
-	SLATE_BEGIN_ARGS(SOpenVDBComponentPicker)
-		: _PackedData()
-		, _ComponentIndex()
-		, _OpenVDBGridComponentInfo()
-	{}
+	SLATE_BEGIN_ARGS(SOpenVDBGridInfoTableRow) {}
+		SLATE_ARGUMENT(TSharedPtr<FOpenVDBGridInfo>, OpenVDBGridInfo)
+	SLATE_END_ARGS()
 
+public:
+	void Construct(const FArguments& InArgs, const TSharedRef<STableViewBase>& OwnerTableView);
+	TSharedRef<SWidget> GenerateWidgetForColumn(const FName& ColumnName) override;
+
+private:
+	TSharedPtr<FOpenVDBGridInfo> OpenVDBGridInfo;
+};
+
+class SOpenVDBComponentPicker : public SCompoundWidget
+{
+public:
+	SLATE_BEGIN_ARGS(SOpenVDBComponentPicker) {}
 		SLATE_ARGUMENT(FSparseVolumeRawSourcePackedData*, PackedData)
 		SLATE_ARGUMENT(uint32, ComponentIndex)
-		SLATE_ARGUMENT(TArray<TSharedPtr<FOpenVDBGridComponentInfo>>*, OpenVDBGridComponentInfo)
+		SLATE_ARGUMENT(const TArray<TSharedPtr<FOpenVDBGridComponentInfo>>*, OpenVDBGridComponentInfo)
 	SLATE_END_ARGS()
 
 public:
@@ -39,23 +51,17 @@ public:
 private:
 	FSparseVolumeRawSourcePackedData* PackedData;
 	uint32 ComponentIndex;
-	TArray<TSharedPtr<FOpenVDBGridComponentInfo>>* OpenVDBGridComponentInfo;
+	const TArray<TSharedPtr<FOpenVDBGridComponentInfo>>* OpenVDBGridComponentInfo;
 	TSharedPtr<SComboBox<TSharedPtr<FOpenVDBGridComponentInfo>>> GridComboBox;
 };
 
 class SOpenVDBPackedDataConfigurator : public SCompoundWidget
 {
 public:
-	SLATE_BEGIN_ARGS(SOpenVDBPackedDataConfigurator)
-		: _PackedData()
-		, _OpenVDBGridComponentInfo()
-		, _OpenVDBSupportedTargetFormats()
-		, _PackedDataName()
-	{}
-
+	SLATE_BEGIN_ARGS(SOpenVDBPackedDataConfigurator) {}
 		SLATE_ARGUMENT(FSparseVolumeRawSourcePackedData*, PackedData)
-		SLATE_ARGUMENT(TArray<TSharedPtr<FOpenVDBGridComponentInfo>>*, OpenVDBGridComponentInfo)
-		SLATE_ARGUMENT(TArray<TSharedPtr<ESparseVolumePackedDataFormat>>*, OpenVDBSupportedTargetFormats)
+		SLATE_ARGUMENT(const TArray<TSharedPtr<FOpenVDBGridComponentInfo>>*, OpenVDBGridComponentInfo)
+		SLATE_ARGUMENT(const TArray<TSharedPtr<ESparseVolumePackedDataFormat>>*, OpenVDBSupportedTargetFormats)
 		SLATE_ARGUMENT(FText, PackedDataName)
 	SLATE_END_ARGS()
 
@@ -66,31 +72,21 @@ public:
 private:
 	FSparseVolumeRawSourcePackedData* PackedData;
 	TSharedPtr<SOpenVDBComponentPicker> ComponentPickers[4];
-	TArray<TSharedPtr<ESparseVolumePackedDataFormat>>* OpenVDBSupportedTargetFormats;
-	TSharedPtr<SComboBox<TSharedPtr<ESparseVolumePackedDataFormat>>>	FormatComboBox;
+	const TArray<TSharedPtr<ESparseVolumePackedDataFormat>>* OpenVDBSupportedTargetFormats;
+	TSharedPtr<SComboBox<TSharedPtr<ESparseVolumePackedDataFormat>>> FormatComboBox;
 	TSharedPtr<SCheckBox> RemapUnormCheckBox;
 };
 
 class SOpenVDBImportWindow : public SCompoundWidget
 {
 public:
-	SLATE_BEGIN_ARGS(SOpenVDBImportWindow)
-		: _PackedDataA()
-		, _PackedDataB()
-		, _OpenVDBGridComponentInfo()
-		, _FileInfoString()
-		, _OpenVDBSupportedTargetFormats()
-		, _WidgetWindow()
-		, _FullPath()
-		, _MaxWindowHeight(0.0f)
-		, _MaxWindowWidth(0.0f)
-	{}
-
+	SLATE_BEGIN_ARGS(SOpenVDBImportWindow) {}
 		SLATE_ARGUMENT(FSparseVolumeRawSourcePackedData*, PackedDataA)
 		SLATE_ARGUMENT(FSparseVolumeRawSourcePackedData*, PackedDataB)
-		SLATE_ARGUMENT(TArray<TSharedPtr<FOpenVDBGridComponentInfo>>*, OpenVDBGridComponentInfo)
-		SLATE_ARGUMENT(FString, FileInfoString)
-		SLATE_ARGUMENT(TArray<TSharedPtr<ESparseVolumePackedDataFormat>>*, OpenVDBSupportedTargetFormats)
+		SLATE_ARGUMENT(int32, NumFoundFiles)
+		SLATE_ARGUMENT(const TArray<TSharedPtr<FOpenVDBGridInfo>>*, OpenVDBGridInfo)
+		SLATE_ARGUMENT(const TArray<TSharedPtr<FOpenVDBGridComponentInfo>>*, OpenVDBGridComponentInfo)
+		SLATE_ARGUMENT(const TArray<TSharedPtr<ESparseVolumePackedDataFormat>>*, OpenVDBSupportedTargetFormats)
 		SLATE_ARGUMENT(TSharedPtr<SWindow>, WidgetWindow)
 		SLATE_ARGUMENT(FText, FullPath)
 		SLATE_ARGUMENT(float, MaxWindowHeight)
@@ -106,18 +102,23 @@ public:
 	bool ShouldImportAsSequence() const;
 
 private:
-	FSparseVolumeRawSourcePackedData*					PackedDataA;
-	FSparseVolumeRawSourcePackedData*					PackedDataB;
-	TArray<TSharedPtr<FOpenVDBGridComponentInfo>>*		OpenVDBGridComponentInfo;
-	TArray<TSharedPtr<ESparseVolumePackedDataFormat>>*	OpenVDBSupportedTargetFormats;
-	TSharedPtr<SOpenVDBPackedDataConfigurator>			PackedDataAConfigurator;
-	TSharedPtr<SOpenVDBPackedDataConfigurator>			PackedDataBConfigurator;
-	TSharedPtr<SCheckBox>								ImportAsSequenceCheckBox;
-	TSharedPtr<SButton>									ImportButton;
-	TWeakPtr<SWindow>									WidgetWindow;
-	bool												bShouldImport;
+	FSparseVolumeRawSourcePackedData* PackedDataA;
+	FSparseVolumeRawSourcePackedData* PackedDataB;
+	FSparseVolumeRawSourcePackedData DefaultAssignmentA;
+	FSparseVolumeRawSourcePackedData DefaultAssignmentB;
+	bool bIsSequence;
+	const TArray<TSharedPtr<FOpenVDBGridInfo>>* OpenVDBGridInfo;
+	const TArray<TSharedPtr<FOpenVDBGridComponentInfo>>* OpenVDBGridComponentInfo;
+	const TArray<TSharedPtr<ESparseVolumePackedDataFormat>>* OpenVDBSupportedTargetFormats;
+	TSharedPtr<SOpenVDBPackedDataConfigurator> PackedDataAConfigurator;
+	TSharedPtr<SOpenVDBPackedDataConfigurator> PackedDataBConfigurator;
+	TSharedPtr<SCheckBox> ImportAsSequenceCheckBox;
+	TSharedPtr<SButton> ImportButton;
+	TWeakPtr<SWindow> WidgetWindow;
+	bool bShouldImport;
 
 	EActiveTimerReturnType SetFocusPostConstruct(double InCurrentTime, float InDeltaTime);
+	TSharedRef<ITableRow> GenerateGridInfoItemRow(TSharedPtr<FOpenVDBGridInfo> Item, const TSharedRef<STableViewBase>& OwnerTable);
 	bool CanImport() const;
 	FReply OnResetToDefaultClick();
 	FText GetImportTypeDisplayText() const;
diff --git a/Engine/Source/Editor/SparseVolumeTexture/Private/SparseVolumeTextureFactory.cpp b/Engine/Source/Editor/SparseVolumeTexture/Private/SparseVolumeTextureFactory.cpp
--- a/Engine/Source/Editor/SparseVolumeTexture/Private/SparseVolumeTextureFactory.cpp
+++ b/Engine/Source/Editor/SparseVolumeTexture/Private/SparseVolumeTextureFactory.cpp
@@ -24,18 +24,147 @@
 
 DEFINE_LOG_CATEGORY_STATIC(LogSparseVolumeTextureFactory, Log, All);
 
+static void ComputeDefaultOpenVDBGridAssignment(const TArray<TSharedPtr<FOpenVDBGridComponentInfo>>& GridComponentInfo, FSparseVolumeRawSourcePackedData* PackedDataA, FSparseVolumeRawSourcePackedData* PackedDataB)
+{
+	FSparseVolumeRawSourcePackedData* PackedData[] = { PackedDataA , PackedDataB };
+	for (FSparseVolumeRawSourcePackedData* Data : PackedData)
+	{
+		Data->Format = ESparseVolumePackedDataFormat::Float32;
+		Data->SourceGridIndex = FUintVector4(INDEX_NONE);
+		Data->SourceComponentIndex = FUintVector4(INDEX_NONE);
+		Data->bRemapInputForUnorm = false;
+	}
+
+	// Assign the components of the input grids to the components of the output SVT.
+	uint32 CurrentOutputPackedData = 0;
+	uint32 CurrentOutputComponent = 0;
+	for (const TSharedPtr<FOpenVDBGridComponentInfo>& GridComponent : GridComponentInfo)
+	{
+		if (GridComponent->Index == INDEX_NONE)
+		{
+			continue;
+		}
+		PackedData[CurrentOutputPackedData]->SourceGridIndex[CurrentOutputComponent] = GridComponent->Index;
+		PackedData[CurrentOutputPackedData]->SourceComponentIndex[CurrentOutputComponent] = GridComponent->ComponentIndex;
+		++CurrentOutputComponent;
+		if (CurrentOutputComponent == 4)
+		{
+			CurrentOutputComponent = 0;
+			++CurrentOutputPackedData;
+			if (CurrentOutputPackedData == 2)
+			{
+				break;
+			}
+		}
+	}
+}
+
+static TArray<FString> FindOpenVDBSequenceFileNames(const FString& Filename)
+{
+	TArray<FString> SequenceFilenames;
+
+	// The file is potentially a sequence if the character before the `.vdb` is a number.
+	const bool bIsFilePotentiallyPartOfASequence = FChar::IsDigit(Filename[Filename.Len() - 5]);
+
+	if (!bIsFilePotentiallyPartOfASequence)
+	{
+		SequenceFilenames.Add(Filename);
+	}
+	else
+	{
+		const FString Path = FPaths::GetPath(Filename);
+		const FString CleanFilename = FPaths::GetCleanFilename(Filename);
+		FString CleanFilenameWithoutSuffix;
+		{
+			const FString CleanFilenameWithoutExt = CleanFilename.LeftChop(4);
+			const int32 LastNonDigitIndex = CleanFilenameWithoutExt.FindLastCharByPredicate([](TCHAR Letter) { return !FChar::IsDigit(Letter); }) + 1;
+			const int32 DigitCount = CleanFilenameWithoutExt.Len() - LastNonDigitIndex;
+			CleanFilenameWithoutSuffix = CleanFilenameWithoutExt.LeftChop(CleanFilenameWithoutExt.Len() - LastNonDigitIndex);
+		}
+
+		// Find all files potentially part of the sequence
+		TArray<FString> PotentialSequenceFilenames;
+		IFileManager::Get().FindFiles(PotentialSequenceFilenames, *Path, TEXT("*.vdb"));
+		PotentialSequenceFilenames = PotentialSequenceFilenames.FilterByPredicate([CleanFilenameWithoutSuffix](const FString& Str) { return Str.StartsWith(CleanFilenameWithoutSuffix); });
+
+		auto GetFilenameNumberSuffix = [](const FString& Filename) -> int32
+		{
+			const FString FilenameWithoutExt = Filename.LeftChop(4);
+			const int32 LastNonDigitIndex = FilenameWithoutExt.FindLastCharByPredicate([](TCHAR Letter) { return !FChar::IsDigit(Letter); }) + 1;
+			const FString NumberSuffixStr = FilenameWithoutExt.RightChop(LastNonDigitIndex);
+
+			int32 Number = INDEX_NONE;
+			if (NumberSuffixStr.IsNumeric())
+			{
+				TTypeFromString<int32>::FromString(Number, *NumberSuffixStr);
+			}
+			return Number;
+		};
+
+		// Find range of number suffixes
+		int32 LowestIndex = INT32_MAX;
+		int32 HighestIndex = INT32_MIN;
+		for (FString& ItemFilename : PotentialSequenceFilenames)
+		{
+			const int32 Index = GetFilenameNumberSuffix(ItemFilename);
+			if (Index == INDEX_NONE)
+			{
+				ItemFilename.Empty();
+				continue;
+			}
+			LowestIndex = FMath::Min(LowestIndex, Index);
+			HighestIndex = FMath::Max(HighestIndex, Index);
+		}
+
+		check(HighestIndex >= LowestIndex);
+
+		// Sort the filenames into the result array
+		SequenceFilenames.SetNum(HighestIndex - LowestIndex + 1);
+		for (const FString& ItemFilename : PotentialSequenceFilenames)
+		{
+			const int32 Index = ItemFilename.IsEmpty() ? INDEX_NONE : GetFilenameNumberSuffix(ItemFilename);
+			if (Index == INDEX_NONE)
+			{
+				continue;
+			}
+			SequenceFilenames[Index - LowestIndex] = Path / ItemFilename;
+		}
+
+		// Chop off any items after finding the first gap
+		for (int32 i = 0; i < SequenceFilenames.Num(); ++i)
+		{
+			if (SequenceFilenames[i].IsEmpty())
+			{
+				SequenceFilenames.SetNum(i);
+				break;
+			}
+		}
+	}
+
+	check(!SequenceFilenames.IsEmpty());
+
+	return SequenceFilenames;
+}
+
 struct FOpenVDBPreviewData
 {
 	TArray<uint8> LoadedFile;
 	TArray<FOpenVDBGridInfo> GridInfo;
+	TArray<TSharedPtr<FOpenVDBGridInfo>> GridInfoPtrs;
 	TArray<TSharedPtr<FOpenVDBGridComponentInfo>> GridComponentInfoPtrs;
-	FString FileInfoString;
+	TArray<FString> SequenceFilenames;
+	FSparseVolumeRawSourcePackedData DefaultGridAssignmentA;
+	FSparseVolumeRawSourcePackedData DefaultGridAssignmentB;
 };
 
 static bool LoadOpenVDBPreviewData(const FString& Filename, FOpenVDBPreviewData* OutPreviewData)
 {
 	FOpenVDBPreviewData& Result = *OutPreviewData;
-	Result = {};
+	check(Result.LoadedFile.IsEmpty());
+	check(Result.GridInfo.IsEmpty());
+	check(Result.GridInfoPtrs.IsEmpty());
+	check(Result.GridComponentInfoPtrs.IsEmpty());
+	check(Result.SequenceFilenames.IsEmpty());
 
 	if (!FFileHelper::LoadFileToArray(Result.LoadedFile, *Filename))
 	{
@@ -67,9 +196,8 @@ static bool LoadOpenVDBPreviewData(const FString& Filename, FOpenVDBPreviewData*
 	bool bFoundSupportedGridType = false;
 	for (const FOpenVDBGridInfo& Grid : Result.GridInfo)
 	{
-		// Append all grids to the string, even if we don't actually support them
-		Result.FileInfoString.Append(Grid.DisplayString);
-		Result.FileInfoString.AppendChar(TEXT('\n'));
+		// Append all grids, even if we don't actually support them
+		Result.GridInfoPtrs.Add(MakeShared<FOpenVDBGridInfo>(Grid));
 
 		if (Grid.Type == EOpenVDBGridType::Unknown || !IsOpenVDBGridValid(Grid, Filename))
 		{
@@ -104,6 +232,10 @@ static bool LoadOpenVDBPreviewData(const FString& Filename, FOpenVDBPreviewData*
 		return false;
 	}
 
+	Result.SequenceFilenames = FindOpenVDBSequenceFileNames(Filename);
+
+	ComputeDefaultOpenVDBGridAssignment(Result.GridComponentInfoPtrs, &Result.DefaultGridAssignmentA, &Result.DefaultGridAssignmentB);
+
 	return true;
 }
 
@@ -112,13 +244,10 @@ struct FOpenVDBImportOptions
 	FSparseVolumeRawSourcePackedData PackedDataA;
 	FSparseVolumeRawSourcePackedData PackedDataB;
 	bool bIsSequence;
-	bool bShouldImport;
 };
 
-static FOpenVDBImportOptions ShowOpenVDBImportWindow(const FString& Filename, const FString& FileInfoString, TArray<TSharedPtr<FOpenVDBGridComponentInfo>>& GridComponentInfo)
+static bool ShowOpenVDBImportWindow(const FString& Filename, const FOpenVDBPreviewData& PreviewData, FOpenVDBImportOptions* OutImportOptions)
 {
-	FOpenVDBImportOptions ImportOptions{};
-
 	TSharedPtr<SWindow> ParentWindow;
 
 	if (FModuleManager::Get().IsModuleLoaded("MainFrame"))
@@ -160,10 +289,11 @@ static FOpenVDBImportOptions ShowOpenVDBImportWindow(const FString& Filename, co
 	Window->SetContent
 	(
 		SAssignNew(OpenVDBOptionWindow, SOpenVDBImportWindow)
-		.PackedDataA(&ImportOptions.PackedDataA)
-		.PackedDataB(&ImportOptions.PackedDataB)
-		.OpenVDBGridComponentInfo(&GridComponentInfo)
-		.FileInfoString(FileInfoString)
+		.PackedDataA(&OutImportOptions->PackedDataA)
+		.PackedDataB(&OutImportOptions->PackedDataB)
+		.NumFoundFiles(PreviewData.SequenceFilenames.Num())
+		.OpenVDBGridInfo(&PreviewData.GridInfoPtrs)
+		.OpenVDBGridComponentInfo(&PreviewData.GridComponentInfoPtrs)
 		.OpenVDBSupportedTargetFormats(&SupportedFormats)
 		.WidgetWindow(Window)
 		.FullPath(FText::FromString(Filename))
@@ -173,10 +303,35 @@ static FOpenVDBImportOptions ShowOpenVDBImportWindow(const FString& Filename, co
 
 	FSlateApplication::Get().AddModalWindow(Window, ParentWindow, false);
 
-	ImportOptions.bShouldImport = OpenVDBOptionWindow->ShouldImport();
-	ImportOptions.bIsSequence = OpenVDBOptionWindow->ShouldImportAsSequence();
+	OutImportOptions->bIsSequence = OpenVDBOptionWindow->ShouldImportAsSequence();
 
-	return ImportOptions;
+	return OpenVDBOptionWindow->ShouldImport();
+}
+
+static bool ValidateImportOptions(const FOpenVDBImportOptions& ImportOptions, const TArray<FOpenVDBGridInfo>& GridInfo)
+{
+	const uint32 NumGrids = (uint32)GridInfo.Num();
+	const FSparseVolumeRawSourcePackedData* PackedData[] = { &ImportOptions.PackedDataA , &ImportOptions.PackedDataB };
+	for (const FSparseVolumeRawSourcePackedData* Data : PackedData)
+	{
+		for (int32 DstCompIdx = 0; DstCompIdx < 4; ++DstCompIdx)
+		{
+			const uint32 SourceGridIndex = Data->SourceGridIndex[DstCompIdx];
+			const uint32 SourceComponentIndex = Data->SourceComponentIndex[DstCompIdx];
+			if (SourceGridIndex != INDEX_NONE)
+			{
+				if (SourceGridIndex >= NumGrids)
+				{
+					return false; // Invalid grid index
+				}
+				if (SourceComponentIndex == INDEX_NONE || SourceComponentIndex >= GridInfo[SourceGridIndex].NumComponents)
+				{
+					return false; // Invalid component index
+				}
+			}
+		}
+	}
+	return true;
 }
 
 USparseVolumeTextureFactory::USparseVolumeTextureFactory(const FObjectInitializer& ObjectInitializer)
@@ -263,6 +418,13 @@ UObject* USparseVolumeTextureFactory::FactoryCreateFile(UClass* InClass, UObject
 	GEditor->GetEditorSubsystem<UImportSubsystem>()->BroadcastAssetPreImport(this, InClass, InParent, InName, Parms);
 	TArray<UObject*> ResultAssets;
 
+	bOutOperationCanceled = false;
+
+	const bool bIsUnattended = (IsAutomatedImport()
+		|| FApp::IsUnattended()
+		|| IsRunningCommandlet()
+		|| GIsRunningUnattendedScript);
+
 	// Load file and get info about each contained grid
 	FOpenVDBPreviewData PreviewData;
 	if (!LoadOpenVDBPreviewData(Filename, &PreviewData))
@@ -270,11 +432,24 @@ UObject* USparseVolumeTextureFactory::FactoryCreateFile(UClass* InClass, UObject
 		return nullptr;
 	}
 
-	// Show dialog for import options
-	FOpenVDBImportOptions ImportOptions = ShowOpenVDBImportWindow(Filename, PreviewData.FileInfoString, PreviewData.GridComponentInfoPtrs);
-	if (!ImportOptions.bShouldImport)
+	FOpenVDBImportOptions ImportOptions;
+	ImportOptions.PackedDataA = PreviewData.DefaultGridAssignmentA;
+	ImportOptions.PackedDataB = PreviewData.DefaultGridAssignmentB;
+	ImportOptions.bIsSequence = PreviewData.SequenceFilenames.Num() > 1;
+	
+	if (!bIsUnattended)
 	{
-		bOutOperationCanceled = true;
+		// Show dialog for import options
+		if (!ShowOpenVDBImportWindow(Filename, PreviewData, &ImportOptions))
+		{
+			bOutOperationCanceled = true;
+			return nullptr;
+		}
+	}
+
+	if (!ValidateImportOptions(ImportOptions, PreviewData.GridInfo))
+	{
+		UE_LOG(LogSparseVolumeTextureFactory, Error, TEXT("Import options are invalid! This is likely due to invalid/out-of-bounds grid or component indices."));
 		return nullptr;
 	}
 
@@ -350,62 +525,7 @@ UObject* USparseVolumeTextureFactory::FactoryCreateFile(UClass* InClass, UObject
 		FName NewName(InName.ToString() + TEXT("VDBAnim"));
 		UAnimatedSparseVolumeTexture* AnimatedSVTexture = NewObject<UAnimatedSparseVolumeTexture>(InParent, UAnimatedSparseVolumeTexture::StaticClass(), NewName, Flags);
 		
-		// The file is potentially a sequence if the character before the `.vdb` is a number.
-		const bool bIsFilePotentiallyPartOfASequence = FChar::IsDigit(Filename[Filename.Len() - 5]);
-
-		// Collect filenames of all files to be imported
-		TArray<FString> SequenceFileNames;
-		if (!bIsFilePotentiallyPartOfASequence)
-		{
-			SequenceFileNames.Add(Filename);
-		}
-		else
-		{
-			const FString FilenameWithoutExt = Filename.LeftChop(4);
-			const int32 LastNonDigitIndex = FilenameWithoutExt.FindLastCharByPredicate([](TCHAR Letter) { return !FChar::IsDigit(Letter); }) + 1;
-			const int32 DigitCount = FilenameWithoutExt.Len() - LastNonDigitIndex;
-			FString FilenameWithoutSuffix = FilenameWithoutExt.LeftChop(FilenameWithoutExt.Len() - LastNonDigitIndex);
-			TCHAR LastDigit = FilenameWithoutExt[FilenameWithoutExt.Len() - 5];
-
-			bool bIndexStartsAtOne = false;
-			auto GetOpenVDBFileNameForFrame = [&](int32 FrameIndex)
-			{
-				FString IndexString = FString::FromInt(FrameIndex + (bIndexStartsAtOne ? 1 : 0));
-				// User must select a frame with index in [0-9] so that we can count leading 0s
-				check(DigitCount == 1 || (DigitCount > 1 && IndexString.Len() <= DigitCount));
-				const int32 MissingLeadingZeroCount = DigitCount - IndexString.Len();
-				const FString StringZero = FString::FromInt(0);
-				for (int32 i = 0; i < MissingLeadingZeroCount; ++i)
-				{
-					IndexString = StringZero + IndexString;
-				}
-				return FString(FilenameWithoutSuffix + IndexString) + TEXT(".vdb");
-			};
-
-			const FString VDBFileAt0 = GetOpenVDBFileNameForFrame(0);
-			const FString VDBFileAt1 = GetOpenVDBFileNameForFrame(1);
-			const bool VDBFileAt0Exists = FPaths::FileExists(VDBFileAt0);
-			const bool VDBFileAt1Exists = FPaths::FileExists(VDBFileAt1);
-			if (!VDBFileAt0Exists && !VDBFileAt1Exists)
-			{
-				UE_LOG(LogSparseVolumeTextureFactory, Error, TEXT("An OpenVDB animated sequence must start at index 0 or 1: %s or %s not found."), *VDBFileAt0, *VDBFileAt1);
-				return nullptr;
-			}
-			bIndexStartsAtOne = !VDBFileAt0Exists;
-
-			// Go over all the frame index and stop at the first missing one.
-			for (int32 FrameIdx = 0; ; ++FrameIdx)
-			{
-				const FString FrameFileName = GetOpenVDBFileNameForFrame(FrameIdx);
-				if (!FPaths::FileExists(FrameFileName))
-				{
-					break;
-				}
-				SequenceFileNames.Add(FrameFileName);
-			}
-		}
-
-		const int32 NumFrames = SequenceFileNames.Num();
+		const int32 NumFrames = PreviewData.SequenceFilenames.Num();
 
 		FScopedSlowTask ImportTask(NumFrames, LOCTEXT("ImportingVDBAnim", "Importing OpenVDB animation"));
 		ImportTask.MakeDialog(true);
@@ -419,7 +539,7 @@ UObject* USparseVolumeTextureFactory::FactoryCreateFile(UClass* InClass, UObject
 		// Load individual frames, check them for compatibility with the first loaded file and append them to the resulting asset
 		for (int32 FrameIdx = 0; FrameIdx < NumFrames; ++FrameIdx)
 		{
-			const FString& FrameFilename = SequenceFileNames[FrameIdx];
+			const FString& FrameFilename = PreviewData.SequenceFilenames[FrameIdx];
 
 			UE_LOG(LogSparseVolumeTextureFactory, Display, TEXT("Loading OpenVDB sequence frame #%i %s."), FrameIdx, *FrameFilename);
 
