SVT uassets now no longer store the original vdb file. Conversion now happens during import, giving the user a progress bar instead of just blocking the editor on PostLoad() when the derived data was generated. Also renamed "PackedData" to "Attributes".

#rnx

diff --git a/Engine/Source/Editor/SparseVolumeTexture/Private/OpenVDBImportOptions.h b/Engine/Source/Editor/SparseVolumeTexture/Private/OpenVDBImportOptions.h
new file mode 100644
--- /dev/null
+++ b/Engine/Source/Editor/SparseVolumeTexture/Private/OpenVDBImportOptions.h
@@ -0,0 +1,32 @@
+// Copyright Epic Games, Inc. All Rights Reserved.
+
+#pragma once
+
+#include "CoreMinimal.h"
+#include "Containers/StaticArray.h"
+
+enum class ESparseVolumeAttributesFormat : uint8
+{
+	Unorm8 = 0,
+	Float16 = 1,
+	Float32 = 2,
+};
+
+struct FOpenVDBSparseVolumeComponentMapping
+{
+	int32 SourceGridIndex = INDEX_NONE;
+	int32 SourceComponentIndex = INDEX_NONE;
+};
+
+struct FOpenVDBSparseVolumeAttributesDesc
+{
+	TStaticArray<FOpenVDBSparseVolumeComponentMapping, 4> Mappings;
+	ESparseVolumeAttributesFormat Format;
+	bool bRemapInputForUnorm; // Maps the input from its minimum and maximum value into the [0-1] range. Clamps to [0-1] otherwise.
+};
+
+struct FOpenVDBImportOptions
+{
+	TStaticArray<FOpenVDBSparseVolumeAttributesDesc, 2> Attributes;
+	bool bIsSequence;
+};
diff --git a/Engine/Source/Editor/SparseVolumeTexture/Private/OpenVDBImportWindow.cpp b/Engine/Source/Editor/SparseVolumeTexture/Private/OpenVDBImportWindow.cpp
--- a/Engine/Source/Editor/SparseVolumeTexture/Private/OpenVDBImportWindow.cpp
+++ b/Engine/Source/Editor/SparseVolumeTexture/Private/OpenVDBImportWindow.cpp
@@ -14,6 +14,7 @@
 #include "Editor.h"
 #include "SparseVolumeTextureOpenVDBUtility.h"
 #include "SparseVolumeTexture/SparseVolumeTexture.h"
+#include "OpenVDBImportOptions.h"
 
 #define LOCTEXT_NAMESPACE "SOpenVDBImportWindow"
 
@@ -27,16 +28,16 @@ static FText GetGridComboBoxItemText(TSharedPtr<FOpenVDBGridComponentInfo> InIte
 	return InItem ? FText::FromString(InItem->DisplayString) : LOCTEXT("NoneGrid", "<None>");
 };
 
-static FText GetFormatComboBoxItemText(TSharedPtr<ESparseVolumePackedDataFormat> InItem)
+static FText GetFormatComboBoxItemText(TSharedPtr<ESparseVolumeAttributesFormat> InItem)
 {
 	const TCHAR* FormatStr = TEXT("<None>");
 	if (InItem)
 	{
 		switch (*InItem)
 		{
-		case ESparseVolumePackedDataFormat::Unorm8: FormatStr = TEXT("8bit unorm"); break;
-		case ESparseVolumePackedDataFormat::Float16: FormatStr = TEXT("16bit float"); break;
-		case ESparseVolumePackedDataFormat::Float32: FormatStr = TEXT("32bit float"); break;
+		case ESparseVolumeAttributesFormat::Unorm8: FormatStr = TEXT("8bit unorm"); break;
+		case ESparseVolumeAttributesFormat::Float16: FormatStr = TEXT("16bit float"); break;
+		case ESparseVolumeAttributesFormat::Float32: FormatStr = TEXT("32bit float"); break;
 		}
 	}
 	return FText::FromString(FormatStr);
@@ -44,10 +45,8 @@ static FText GetFormatComboBoxItemText(TSharedPtr<ESparseVolumePackedDataFormat>
 
 void SOpenVDBImportWindow::Construct(const FArguments& InArgs)
 {
-	PackedDataA = InArgs._PackedDataA;
-	PackedDataB = InArgs._PackedDataB;
-	DefaultAssignmentA = *PackedDataA;
-	DefaultAssignmentB = *PackedDataB;
+	ImportOptions = InArgs._ImportOptions;
+	DefaultImportOptions = InArgs._DefaultImportOptions;
 	bIsSequence = InArgs._NumFoundFiles > 1;
 	OpenVDBGridInfo = InArgs._OpenVDBGridInfo;
 	OpenVDBGridComponentInfo = InArgs._OpenVDBGridComponentInfo;
@@ -169,21 +168,21 @@ void SOpenVDBImportWindow::Construct(const FArguments& InArgs)
 		.AutoHeight()
 		.Padding(2)
 		[
-			SAssignNew(PackedDataAConfigurator, SOpenVDBPackedDataConfigurator)
-			.PackedData(PackedDataA)
+			SAssignNew(AttributesAConfigurator, SOpenVDBAttributesConfigurator)
+			.AttributesDesc(&ImportOptions->Attributes[0])
 			.OpenVDBGridComponentInfo(OpenVDBGridComponentInfo)
 			.OpenVDBSupportedTargetFormats(OpenVDBSupportedTargetFormats)
-			.PackedDataName(LOCTEXT("OpenVDBImportWindow_PackedDataA", "Packed Data A"))
+			.AttributesName(LOCTEXT("OpenVDBImportWindow_AttributesA", "Attributes A"))
 		]
 		+ SVerticalBox::Slot()
 		.AutoHeight()
 		.Padding(2)
 		[
-			SAssignNew(PackedDataBConfigurator, SOpenVDBPackedDataConfigurator)
-			.PackedData(PackedDataB)
+			SAssignNew(AttributesBConfigurator, SOpenVDBAttributesConfigurator)
+			.AttributesDesc(&ImportOptions->Attributes[1])
 			.OpenVDBGridComponentInfo(OpenVDBGridComponentInfo)
 			.OpenVDBSupportedTargetFormats(OpenVDBSupportedTargetFormats)
-			.PackedDataName(LOCTEXT("OpenVDBImportWindow_PackedDataB", "Packed Data B"))
+			.AttributesName(LOCTEXT("OpenVDBImportWindow_AttributesB", "Attributes B"))
 		]
 		+ SVerticalBox::Slot()
 		.AutoHeight()
@@ -304,18 +303,14 @@ TSharedRef<ITableRow> SOpenVDBImportWindow::GenerateGridInfoItemRow(TSharedPtr<F
 
 bool SOpenVDBImportWindow::CanImport() const
 {
-	const FUintVector4& GridIndicesA = PackedDataA->SourceGridIndex;
-	const FUintVector4& ComponentIndicesA = PackedDataA->SourceComponentIndex;
-	const FUintVector4& GridIndicesB = PackedDataB->SourceGridIndex;
-	const FUintVector4& ComponentIndicesB = PackedDataB->SourceComponentIndex;
-
-	for (uint32 i = 0; i < 4; ++i)
+	for (const FOpenVDBSparseVolumeAttributesDesc& AttributesDesc : ImportOptions->Attributes)
 	{
-		const bool bGridAValid = GridIndicesA[i] != INDEX_NONE && ComponentIndicesA[i] != INDEX_NONE;
-		const bool bGridBValid = GridIndicesB[i] != INDEX_NONE && ComponentIndicesB[i] != INDEX_NONE;
-		if (bGridAValid || bGridBValid)
+		for (const FOpenVDBSparseVolumeComponentMapping& Mapping : AttributesDesc.Mappings)
 		{
-			return true;
+			if (Mapping.SourceGridIndex != INDEX_NONE && Mapping.SourceComponentIndex != INDEX_NONE)
+			{
+				return true;
+			}
 		}
 	}
 	return false;
@@ -336,17 +331,16 @@ void SOpenVDBImportWindow::SetDefaultGridAssignment()
 {
 	check(OpenVDBGridComponentInfo);
 
-	*PackedDataA = DefaultAssignmentA;
-	*PackedDataB = DefaultAssignmentB;
+	*ImportOptions = *DefaultImportOptions;
 
 	ImportAsSequenceCheckBox->SetIsChecked(bIsSequence ? ECheckBoxState::Checked : ECheckBoxState::Unchecked);
-	PackedDataAConfigurator->RefreshUIFromData();
-	PackedDataBConfigurator->RefreshUIFromData();
+	AttributesAConfigurator->RefreshUIFromData();
+	AttributesBConfigurator->RefreshUIFromData();
 }
 
 void SOpenVDBComponentPicker::Construct(const FArguments& InArgs)
 {
-	PackedData = InArgs._PackedData;
+	AttributesDesc = InArgs._AttributesDesc;
 	ComponentIndex = InArgs._ComponentIndex;
 	OpenVDBGridComponentInfo = InArgs._OpenVDBGridComponentInfo;
 	
@@ -383,16 +377,8 @@ void SOpenVDBComponentPicker::Construct(const FArguments& InArgs)
 					})
 					.OnSelectionChanged_Lambda([this](TSharedPtr<FOpenVDBGridComponentInfo> InItem, ESelectInfo::Type)
 					{
-						if (InItem)
-						{
-							PackedData->SourceGridIndex[ComponentIndex] = InItem->Index;
-							PackedData->SourceComponentIndex[ComponentIndex] = InItem->ComponentIndex;
-						}
-						else
-						{
-							PackedData->SourceGridIndex[ComponentIndex] = INDEX_NONE;
-							PackedData->SourceComponentIndex[ComponentIndex] = INDEX_NONE;
-						}
+						AttributesDesc->Mappings[ComponentIndex].SourceGridIndex = InItem ? InItem->Index : INDEX_NONE;
+						AttributesDesc->Mappings[ComponentIndex].SourceComponentIndex = InItem ? InItem->ComponentIndex : INDEX_NONE;
 					})
 					[
 						SNew(STextBlock)
@@ -410,7 +396,7 @@ void SOpenVDBComponentPicker::RefreshUIFromData()
 {
 	for (const TSharedPtr<FOpenVDBGridComponentInfo>& Grid : *OpenVDBGridComponentInfo)
 	{
-		if (Grid->Index == PackedData->SourceGridIndex[ComponentIndex] && Grid->ComponentIndex == PackedData->SourceComponentIndex[ComponentIndex])
+		if (Grid->Index == AttributesDesc->Mappings[ComponentIndex].SourceGridIndex && Grid->ComponentIndex == AttributesDesc->Mappings[ComponentIndex].SourceComponentIndex)
 		{
 			GridComboBox->SetSelectedItem(Grid);
 			break;
@@ -418,16 +404,16 @@ void SOpenVDBComponentPicker::RefreshUIFromData()
 	}
 }
 
-void SOpenVDBPackedDataConfigurator::Construct(const FArguments& InArgs)
+void SOpenVDBAttributesConfigurator::Construct(const FArguments& InArgs)
 {
-	PackedData = InArgs._PackedData;
+	AttributesDesc = InArgs._AttributesDesc;
 	OpenVDBSupportedTargetFormats = InArgs._OpenVDBSupportedTargetFormats;
 
 	for (uint32 ComponentIndex = 0; ComponentIndex < 4; ++ComponentIndex)
 	{
 		ComponentPickers[ComponentIndex] =
 			SNew(SOpenVDBComponentPicker)
-			.PackedData(PackedData)
+			.AttributesDesc(AttributesDesc)
 			.ComponentIndex(ComponentIndex)
 			.OpenVDBGridComponentInfo(InArgs._OpenVDBGridComponentInfo);
 	}
@@ -448,7 +434,7 @@ void SOpenVDBPackedDataConfigurator::Construct(const FArguments& InArgs)
 				.Padding(2.0f)
 				[
 					SNew(STextBlock)
-					.Text(InArgs._PackedDataName)
+					.Text(InArgs._AttributesName)
 				]
 
 				+ SHorizontalBox::Slot()
@@ -459,16 +445,16 @@ void SOpenVDBPackedDataConfigurator::Construct(const FArguments& InArgs)
 					SNew(SBox)
 					.WidthOverride(50.0f)
 					[
-						SAssignNew(FormatComboBox, SComboBox<TSharedPtr<ESparseVolumePackedDataFormat>>)
+						SAssignNew(FormatComboBox, SComboBox<TSharedPtr<ESparseVolumeAttributesFormat>>)
 						.OptionsSource(OpenVDBSupportedTargetFormats)
-						.OnGenerateWidget_Lambda([](TSharedPtr<ESparseVolumePackedDataFormat> InItem)
+						.OnGenerateWidget_Lambda([](TSharedPtr<ESparseVolumeAttributesFormat> InItem)
 						{
 							return SNew(STextBlock)
 							.Text(GetFormatComboBoxItemText(InItem));
 						})
-						.OnSelectionChanged_Lambda([this](TSharedPtr<ESparseVolumePackedDataFormat> InItem, ESelectInfo::Type)
+						.OnSelectionChanged_Lambda([this](TSharedPtr<ESparseVolumeAttributesFormat> InItem, ESelectInfo::Type)
 						{
-							PackedData->Format = InItem ? *InItem : ESparseVolumePackedDataFormat::Float32;
+							AttributesDesc->Format = InItem ? *InItem : ESparseVolumeAttributesFormat::Float32;
 						})
 						[
 							SNew(STextBlock)
@@ -489,7 +475,7 @@ void SOpenVDBPackedDataConfigurator::Construct(const FArguments& InArgs)
 					.Text(LOCTEXT("UnormRemapCheckBoxLabel", "Unorm Remap"))
 					.IsEnabled_Lambda([this]()
 					{
-						return PackedData->Format == ESparseVolumePackedDataFormat::Unorm8;
+						return AttributesDesc->Format == ESparseVolumeAttributesFormat::Unorm8;
 					})
 				]
 				+ SHorizontalBox::Slot()
@@ -501,11 +487,11 @@ void SOpenVDBPackedDataConfigurator::Construct(const FArguments& InArgs)
 					SAssignNew(RemapUnormCheckBox, SCheckBox)
 					.OnCheckStateChanged_Lambda([this](ECheckBoxState CheckBoxState)
 					{
-						PackedData->bRemapInputForUnorm = CheckBoxState == ECheckBoxState::Checked;
+						AttributesDesc->bRemapInputForUnorm = CheckBoxState == ECheckBoxState::Checked;
 					})
 					.IsEnabled_Lambda([this]()
 					{
-						return PackedData->Format == ESparseVolumePackedDataFormat::Unorm8;
+						return AttributesDesc->Format == ESparseVolumeAttributesFormat::Unorm8;
 					})
 					.ToolTipText(LOCTEXT("UnormRemapCheckBoxTooltip", "Remaps input values for unorm formats into the [0-1] range instead of clamping values outside this range."))
 					.IsChecked(false)
@@ -550,11 +536,11 @@ void SOpenVDBPackedDataConfigurator::Construct(const FArguments& InArgs)
 		];
 }
 
-void SOpenVDBPackedDataConfigurator::RefreshUIFromData()
+void SOpenVDBAttributesConfigurator::RefreshUIFromData()
 {
 	for (auto& Format : *OpenVDBSupportedTargetFormats)
 	{
-		if (*Format == PackedData->Format)
+		if (*Format == AttributesDesc->Format)
 		{
 			FormatComboBox->SetSelectedItem(Format);
 			break;
@@ -564,7 +550,7 @@ void SOpenVDBPackedDataConfigurator::RefreshUIFromData()
 	{
 		ComponentPickers[i]->RefreshUIFromData();
 	}
-	RemapUnormCheckBox->SetIsChecked(PackedData->bRemapInputForUnorm);
+	RemapUnormCheckBox->SetIsChecked(AttributesDesc->bRemapInputForUnorm);
 }
 
 void SOpenVDBGridInfoTableRow::Construct(const FArguments& InArgs, const TSharedRef<STableViewBase>& OwnerTableView)
diff --git a/Engine/Source/Editor/SparseVolumeTexture/Private/OpenVDBImportWindow.h b/Engine/Source/Editor/SparseVolumeTexture/Private/OpenVDBImportWindow.h
--- a/Engine/Source/Editor/SparseVolumeTexture/Private/OpenVDBImportWindow.h
+++ b/Engine/Source/Editor/SparseVolumeTexture/Private/OpenVDBImportWindow.h
@@ -10,7 +10,9 @@
 #include "SparseVolumeTexture/SparseVolumeTexture.h"
 
 struct FOpenVDBGridInfo;
-enum class ESparseVolumePackedDataFormat : uint8;
+struct FOpenVDBImportOptions;
+struct FOpenVDBSparseVolumeAttributesDesc;
+enum class ESparseVolumeAttributesFormat : uint8;
 
 struct FOpenVDBGridComponentInfo
 {
@@ -39,7 +41,7 @@ class SOpenVDBComponentPicker : public SCompoundWidget
 {
 public:
 	SLATE_BEGIN_ARGS(SOpenVDBComponentPicker) {}
-		SLATE_ARGUMENT(FSparseVolumeRawSourcePackedData*, PackedData)
+		SLATE_ARGUMENT(FOpenVDBSparseVolumeAttributesDesc*, AttributesDesc)
 		SLATE_ARGUMENT(uint32, ComponentIndex)
 		SLATE_ARGUMENT(const TArray<TSharedPtr<FOpenVDBGridComponentInfo>>*, OpenVDBGridComponentInfo)
 	SLATE_END_ARGS()
@@ -49,20 +51,20 @@ public:
 	void RefreshUIFromData();
 
 private:
-	FSparseVolumeRawSourcePackedData* PackedData;
+	FOpenVDBSparseVolumeAttributesDesc* AttributesDesc;
 	uint32 ComponentIndex;
 	const TArray<TSharedPtr<FOpenVDBGridComponentInfo>>* OpenVDBGridComponentInfo;
 	TSharedPtr<SComboBox<TSharedPtr<FOpenVDBGridComponentInfo>>> GridComboBox;
 };
 
-class SOpenVDBPackedDataConfigurator : public SCompoundWidget
+class SOpenVDBAttributesConfigurator : public SCompoundWidget
 {
 public:
-	SLATE_BEGIN_ARGS(SOpenVDBPackedDataConfigurator) {}
-		SLATE_ARGUMENT(FSparseVolumeRawSourcePackedData*, PackedData)
+	SLATE_BEGIN_ARGS(SOpenVDBAttributesConfigurator) {}
+		SLATE_ARGUMENT(FOpenVDBSparseVolumeAttributesDesc*, AttributesDesc)
 		SLATE_ARGUMENT(const TArray<TSharedPtr<FOpenVDBGridComponentInfo>>*, OpenVDBGridComponentInfo)
-		SLATE_ARGUMENT(const TArray<TSharedPtr<ESparseVolumePackedDataFormat>>*, OpenVDBSupportedTargetFormats)
-		SLATE_ARGUMENT(FText, PackedDataName)
+		SLATE_ARGUMENT(const TArray<TSharedPtr<ESparseVolumeAttributesFormat>>*, OpenVDBSupportedTargetFormats)
+		SLATE_ARGUMENT(FText, AttributesName)
 	SLATE_END_ARGS()
 
 public:
@@ -70,10 +72,10 @@ public:
 	void RefreshUIFromData();
 
 private:
-	FSparseVolumeRawSourcePackedData* PackedData;
+	FOpenVDBSparseVolumeAttributesDesc* AttributesDesc;
 	TSharedPtr<SOpenVDBComponentPicker> ComponentPickers[4];
-	const TArray<TSharedPtr<ESparseVolumePackedDataFormat>>* OpenVDBSupportedTargetFormats;
-	TSharedPtr<SComboBox<TSharedPtr<ESparseVolumePackedDataFormat>>> FormatComboBox;
+	const TArray<TSharedPtr<ESparseVolumeAttributesFormat>>* OpenVDBSupportedTargetFormats;
+	TSharedPtr<SComboBox<TSharedPtr<ESparseVolumeAttributesFormat>>> FormatComboBox;
 	TSharedPtr<SCheckBox> RemapUnormCheckBox;
 };
 
@@ -81,12 +83,12 @@ class SOpenVDBImportWindow : public SCompoundWidget
 {
 public:
 	SLATE_BEGIN_ARGS(SOpenVDBImportWindow) {}
-		SLATE_ARGUMENT(FSparseVolumeRawSourcePackedData*, PackedDataA)
-		SLATE_ARGUMENT(FSparseVolumeRawSourcePackedData*, PackedDataB)
+		SLATE_ARGUMENT(FOpenVDBImportOptions*, ImportOptions)
+		SLATE_ARGUMENT(const FOpenVDBImportOptions*, DefaultImportOptions)
 		SLATE_ARGUMENT(int32, NumFoundFiles)
 		SLATE_ARGUMENT(const TArray<TSharedPtr<FOpenVDBGridInfo>>*, OpenVDBGridInfo)
 		SLATE_ARGUMENT(const TArray<TSharedPtr<FOpenVDBGridComponentInfo>>*, OpenVDBGridComponentInfo)
-		SLATE_ARGUMENT(const TArray<TSharedPtr<ESparseVolumePackedDataFormat>>*, OpenVDBSupportedTargetFormats)
+		SLATE_ARGUMENT(const TArray<TSharedPtr<ESparseVolumeAttributesFormat>>*, OpenVDBSupportedTargetFormats)
 		SLATE_ARGUMENT(TSharedPtr<SWindow>, WidgetWindow)
 		SLATE_ARGUMENT(FText, FullPath)
 		SLATE_ARGUMENT(float, MaxWindowHeight)
@@ -102,16 +104,14 @@ public:
 	bool ShouldImportAsSequence() const;
 
 private:
-	FSparseVolumeRawSourcePackedData* PackedDataA;
-	FSparseVolumeRawSourcePackedData* PackedDataB;
-	FSparseVolumeRawSourcePackedData DefaultAssignmentA;
-	FSparseVolumeRawSourcePackedData DefaultAssignmentB;
+	FOpenVDBImportOptions* ImportOptions;
+	const FOpenVDBImportOptions* DefaultImportOptions;
 	bool bIsSequence;
 	const TArray<TSharedPtr<FOpenVDBGridInfo>>* OpenVDBGridInfo;
 	const TArray<TSharedPtr<FOpenVDBGridComponentInfo>>* OpenVDBGridComponentInfo;
-	const TArray<TSharedPtr<ESparseVolumePackedDataFormat>>* OpenVDBSupportedTargetFormats;
-	TSharedPtr<SOpenVDBPackedDataConfigurator> PackedDataAConfigurator;
-	TSharedPtr<SOpenVDBPackedDataConfigurator> PackedDataBConfigurator;
+	const TArray<TSharedPtr<ESparseVolumeAttributesFormat>>* OpenVDBSupportedTargetFormats;
+	TSharedPtr<SOpenVDBAttributesConfigurator> AttributesAConfigurator;
+	TSharedPtr<SOpenVDBAttributesConfigurator> AttributesBConfigurator;
 	TSharedPtr<SCheckBox> ImportAsSequenceCheckBox;
 	TSharedPtr<SButton> ImportButton;
 	TWeakPtr<SWindow> WidgetWindow;
diff --git a/Engine/Source/Editor/SparseVolumeTexture/Private/SparseVolumeTextureFactory.cpp b/Engine/Source/Editor/SparseVolumeTexture/Private/SparseVolumeTextureFactory.cpp
--- a/Engine/Source/Editor/SparseVolumeTexture/Private/SparseVolumeTextureFactory.cpp
+++ b/Engine/Source/Editor/SparseVolumeTexture/Private/SparseVolumeTextureFactory.cpp
@@ -8,6 +8,7 @@
 
 #include "SparseVolumeTextureOpenVDB.h"
 #include "SparseVolumeTextureOpenVDBUtility.h"
+#include "OpenVDBImportOptions.h"
 
 #include "Serialization/EditorBulkDataWriter.h"
 #include "Misc/Paths.h"
@@ -24,39 +25,42 @@
 
 DEFINE_LOG_CATEGORY_STATIC(LogSparseVolumeTextureFactory, Log, All);
 
-static void ComputeDefaultOpenVDBGridAssignment(const TArray<TSharedPtr<FOpenVDBGridComponentInfo>>& GridComponentInfo, FSparseVolumeRawSourcePackedData* PackedDataA, FSparseVolumeRawSourcePackedData* PackedDataB)
+static void ComputeDefaultOpenVDBGridAssignment(const TArray<TSharedPtr<FOpenVDBGridComponentInfo>>& GridComponentInfo, int32 NumFiles, FOpenVDBImportOptions* ImportOptions)
 {
-	FSparseVolumeRawSourcePackedData* PackedData[] = { PackedDataA , PackedDataB };
-	for (FSparseVolumeRawSourcePackedData* Data : PackedData)
+	for (FOpenVDBSparseVolumeAttributesDesc& AttributesDesc : ImportOptions->Attributes)
 	{
-		Data->Format = ESparseVolumePackedDataFormat::Float32;
-		Data->SourceGridIndex = FUintVector4(INDEX_NONE);
-		Data->SourceComponentIndex = FUintVector4(INDEX_NONE);
-		Data->bRemapInputForUnorm = false;
+		for (FOpenVDBSparseVolumeComponentMapping& Mapping : AttributesDesc.Mappings)
+		{
+			Mapping.SourceGridIndex = INDEX_NONE;
+			Mapping.SourceComponentIndex = INDEX_NONE;
+		}
+		AttributesDesc.Format = ESparseVolumeAttributesFormat::Float32;
+		AttributesDesc.bRemapInputForUnorm = false;
 	}
 
 	// Assign the components of the input grids to the components of the output SVT.
-	uint32 CurrentOutputPackedData = 0;
-	uint32 CurrentOutputComponent = 0;
+	uint32 DstAttributesIdx = 0;
+	uint32 DstComponentIdx = 0;
 	for (const TSharedPtr<FOpenVDBGridComponentInfo>& GridComponent : GridComponentInfo)
 	{
 		if (GridComponent->Index == INDEX_NONE)
 		{
 			continue;
 		}
-		PackedData[CurrentOutputPackedData]->SourceGridIndex[CurrentOutputComponent] = GridComponent->Index;
-		PackedData[CurrentOutputPackedData]->SourceComponentIndex[CurrentOutputComponent] = GridComponent->ComponentIndex;
-		++CurrentOutputComponent;
-		if (CurrentOutputComponent == 4)
+		ImportOptions->Attributes[DstAttributesIdx].Mappings[DstComponentIdx].SourceGridIndex = GridComponent->Index;
+		ImportOptions->Attributes[DstAttributesIdx].Mappings[DstComponentIdx].SourceComponentIndex = GridComponent->ComponentIndex;
+		++DstComponentIdx;
+		if (DstComponentIdx == 4)
 		{
-			CurrentOutputComponent = 0;
-			++CurrentOutputPackedData;
-			if (CurrentOutputPackedData == 2)
+			DstComponentIdx = 0;
+			++DstAttributesIdx;
+			if (DstAttributesIdx == 2)
 			{
 				break;
 			}
 		}
 	}
+	ImportOptions->bIsSequence = NumFiles > 1;
 }
 
 static TArray<FString> FindOpenVDBSequenceFileNames(const FString& Filename)
@@ -153,8 +157,7 @@ struct FOpenVDBPreviewData
 	TArray<TSharedPtr<FOpenVDBGridInfo>> GridInfoPtrs;
 	TArray<TSharedPtr<FOpenVDBGridComponentInfo>> GridComponentInfoPtrs;
 	TArray<FString> SequenceFilenames;
-	FSparseVolumeRawSourcePackedData DefaultGridAssignmentA;
-	FSparseVolumeRawSourcePackedData DefaultGridAssignmentB;
+	FOpenVDBImportOptions DefaultImportOptions;
 };
 
 static bool LoadOpenVDBPreviewData(const FString& Filename, FOpenVDBPreviewData* OutPreviewData)
@@ -234,18 +237,11 @@ static bool LoadOpenVDBPreviewData(const FString& Filename, FOpenVDBPreviewData*
 
 	Result.SequenceFilenames = FindOpenVDBSequenceFileNames(Filename);
 
-	ComputeDefaultOpenVDBGridAssignment(Result.GridComponentInfoPtrs, &Result.DefaultGridAssignmentA, &Result.DefaultGridAssignmentB);
+	ComputeDefaultOpenVDBGridAssignment(Result.GridComponentInfoPtrs, Result.SequenceFilenames.Num(), &Result.DefaultImportOptions);
 
 	return true;
 }
 
-struct FOpenVDBImportOptions
-{
-	FSparseVolumeRawSourcePackedData PackedDataA;
-	FSparseVolumeRawSourcePackedData PackedDataB;
-	bool bIsSequence;
-};
-
 static bool ShowOpenVDBImportWindow(const FString& Filename, const FOpenVDBPreviewData& PreviewData, FOpenVDBImportOptions* OutImportOptions)
 {
 	TSharedPtr<SWindow> ParentWindow;
@@ -278,19 +274,19 @@ static bool ShowOpenVDBImportWindow(const FString& Filename, const FOpenVDBPrevi
 		.ClientSize(ImportWindowSize)
 		.ScreenPosition(WindowPosition);
 
-	TArray<TSharedPtr<ESparseVolumePackedDataFormat>> SupportedFormats =
+	TArray<TSharedPtr<ESparseVolumeAttributesFormat>> SupportedFormats =
 	{
-		MakeShared<ESparseVolumePackedDataFormat>(ESparseVolumePackedDataFormat::Float32),
-		MakeShared<ESparseVolumePackedDataFormat>(ESparseVolumePackedDataFormat::Float16),
-		MakeShared<ESparseVolumePackedDataFormat>(ESparseVolumePackedDataFormat::Unorm8)
+		MakeShared<ESparseVolumeAttributesFormat>(ESparseVolumeAttributesFormat::Float32),
+		MakeShared<ESparseVolumeAttributesFormat>(ESparseVolumeAttributesFormat::Float16),
+		MakeShared<ESparseVolumeAttributesFormat>(ESparseVolumeAttributesFormat::Unorm8)
 	};
 
 	TSharedPtr<SOpenVDBImportWindow> OpenVDBOptionWindow;
 	Window->SetContent
 	(
 		SAssignNew(OpenVDBOptionWindow, SOpenVDBImportWindow)
-		.PackedDataA(&OutImportOptions->PackedDataA)
-		.PackedDataB(&OutImportOptions->PackedDataB)
+		.ImportOptions(OutImportOptions)
+		.DefaultImportOptions(&PreviewData.DefaultImportOptions)
 		.NumFoundFiles(PreviewData.SequenceFilenames.Num())
 		.OpenVDBGridInfo(&PreviewData.GridInfoPtrs)
 		.OpenVDBGridComponentInfo(&PreviewData.GridComponentInfoPtrs)
@@ -310,21 +306,21 @@ static bool ShowOpenVDBImportWindow(const FString& Filename, const FOpenVDBPrevi
 
 static bool ValidateImportOptions(const FOpenVDBImportOptions& ImportOptions, const TArray<FOpenVDBGridInfo>& GridInfo)
 {
-	const uint32 NumGrids = (uint32)GridInfo.Num();
-	const FSparseVolumeRawSourcePackedData* PackedData[] = { &ImportOptions.PackedDataA , &ImportOptions.PackedDataB };
-	for (const FSparseVolumeRawSourcePackedData* Data : PackedData)
+	const int32 NumGrids = GridInfo.Num();
+
+	for (const FOpenVDBSparseVolumeAttributesDesc& AttributesDesc : ImportOptions.Attributes)
 	{
-		for (int32 DstCompIdx = 0; DstCompIdx < 4; ++DstCompIdx)
+		for (const FOpenVDBSparseVolumeComponentMapping& Mapping : AttributesDesc.Mappings)
 		{
-			const uint32 SourceGridIndex = Data->SourceGridIndex[DstCompIdx];
-			const uint32 SourceComponentIndex = Data->SourceComponentIndex[DstCompIdx];
-			if (SourceGridIndex != INDEX_NONE)
+			const int32 SourceGridIndex = Mapping.SourceGridIndex;
+			const int32 SourceComponentIndex = Mapping.SourceComponentIndex;
+			if (Mapping.SourceGridIndex != INDEX_NONE)
 			{
 				if (SourceGridIndex >= NumGrids)
 				{
 					return false; // Invalid grid index
 				}
-				if (SourceComponentIndex == INDEX_NONE || SourceComponentIndex >= GridInfo[SourceGridIndex].NumComponents)
+				if (SourceComponentIndex == INDEX_NONE || SourceComponentIndex >= (int32)GridInfo[SourceGridIndex].NumComponents)
 				{
 					return false; // Invalid component index
 				}
@@ -432,10 +428,7 @@ UObject* USparseVolumeTextureFactory::FactoryCreateFile(UClass* InClass, UObject
 		return nullptr;
 	}
 
-	FOpenVDBImportOptions ImportOptions;
-	ImportOptions.PackedDataA = PreviewData.DefaultGridAssignmentA;
-	ImportOptions.PackedDataB = PreviewData.DefaultGridAssignmentB;
-	ImportOptions.bIsSequence = PreviewData.SequenceFilenames.Num() > 1;
+	FOpenVDBImportOptions ImportOptions = PreviewData.DefaultImportOptions;
 	
 	if (!bIsUnattended)
 	{
@@ -456,15 +449,13 @@ UObject* USparseVolumeTextureFactory::FactoryCreateFile(UClass* InClass, UObject
 	// Utility function for computing the bounding box encompassing the bounds of all frames in the SVT.
 	auto ExpandVolumeBounds = [](const FOpenVDBImportOptions& ImportOptions, const TArray<FOpenVDBGridInfo>& GridInfoArray, FBox& VolumeBounds)
 	{
-		const FSparseVolumeRawSourcePackedData* PackedData[] = { &ImportOptions.PackedDataA, &ImportOptions.PackedDataB };
-		for (uint32 PackedDataIdx = 0; PackedDataIdx < 2; ++PackedDataIdx)
+		for (const FOpenVDBSparseVolumeAttributesDesc& Attributes : ImportOptions.Attributes)
 		{
-			for (uint32 CompIdx = 0; CompIdx < 4; ++CompIdx)
+			for (const FOpenVDBSparseVolumeComponentMapping& Mapping : Attributes.Mappings)
 			{
-				const uint32 GridIdx = PackedData[PackedDataIdx]->SourceGridIndex[CompIdx];
-				if (GridIdx != INDEX_NONE)
+				if (Mapping.SourceGridIndex != INDEX_NONE)
 				{
-					const FOpenVDBGridInfo& GridInfo = GridInfoArray[GridIdx];
+					const FOpenVDBGridInfo& GridInfo = GridInfoArray[Mapping.SourceGridIndex];
 					VolumeBounds.Min.X = FMath::Min(VolumeBounds.Min.X, GridInfo.VolumeActiveAABBMin.X);
 					VolumeBounds.Min.Y = FMath::Min(VolumeBounds.Min.Y, GridInfo.VolumeActiveAABBMin.Y);
 					VolumeBounds.Min.Z = FMath::Min(VolumeBounds.Min.Z, GridInfo.VolumeActiveAABBMin.Z);
@@ -493,10 +484,24 @@ UObject* USparseVolumeTextureFactory::FactoryCreateFile(UClass* InClass, UObject
 		StaticSVTexture->VolumeBounds = VolumeBounds;
 
 		FSparseVolumeRawSource SparseVolumeRawSource{};
-		SparseVolumeRawSource.PackedDataA = ImportOptions.PackedDataA;
-		SparseVolumeRawSource.PackedDataB = ImportOptions.PackedDataB;
-		SparseVolumeRawSource.SourceAssetFile = MoveTemp(PreviewData.LoadedFile);
-		PreviewData.LoadedFile.Reset();
+
+		FOpenVDBToSVTConversionResult SVTResult;
+		SVTResult.Header = &SparseVolumeRawSource.Header;
+		SVTResult.PageTable = &SparseVolumeRawSource.PageTable;
+		SVTResult.PhysicalTileDataA = &SparseVolumeRawSource.PhysicalTileDataA;
+		SVTResult.PhysicalTileDataB = &SparseVolumeRawSource.PhysicalTileDataB;
+
+		const bool bConversionSuccess = ConvertOpenVDBToSparseVolumeTexture(
+			PreviewData.LoadedFile,
+			ImportOptions,
+			&SVTResult,
+			false, FVector::Zero(), FVector::Zero());
+
+		if (!bConversionSuccess)
+		{
+			UE_LOG(LogSparseVolumeTextureFactory, Error, TEXT("Failed to convert OpenVDB file to SparseVolumeTexture: %s"), *Filename);
+			return nullptr;
+		}
 
 		// Serialize the raw source data into the asset object.
 		{
@@ -506,6 +511,7 @@ UObject* USparseVolumeTextureFactory::FactoryCreateFile(UClass* InClass, UObject
 
 		if (ImportTask.ShouldCancel())
 		{
+			bOutOperationCanceled = true;
 			return nullptr;
 		}
 		ImportTask.EnterProgressFrame(1.0f, LOCTEXT("ConvertingVDBStatic", "Converting static OpenVDB"));
@@ -556,12 +562,11 @@ UObject* USparseVolumeTextureFactory::FactoryCreateFile(UClass* InClass, UObject
 			}
 
 			// Sanity check for compatibility
-			const FSparseVolumeRawSourcePackedData* PackedData[] = { &ImportOptions.PackedDataA, &ImportOptions.PackedDataB };
-			for (uint32 PackedDataIdx = 0; PackedDataIdx < 2; ++PackedDataIdx)
+			for (const FOpenVDBSparseVolumeAttributesDesc& AttributesDesc : ImportOptions.Attributes)
 			{
-				for (uint32 DstCompIdx = 0; DstCompIdx < 4; ++DstCompIdx)
+				for (const FOpenVDBSparseVolumeComponentMapping& Mapping : AttributesDesc.Mappings)
 				{
-					const uint32 SourceGridIndex = PackedData[PackedDataIdx]->SourceGridIndex[DstCompIdx];
+					const uint32 SourceGridIndex = Mapping.SourceGridIndex;
 					if (SourceGridIndex != INDEX_NONE)
 					{
 						if ((int32)SourceGridIndex >= FrameGridInfo.Num())
@@ -580,12 +585,25 @@ UObject* USparseVolumeTextureFactory::FactoryCreateFile(UClass* InClass, UObject
 				}
 			}
 
-			ExpandVolumeBounds(ImportOptions, FrameGridInfo, VolumeBounds);
-
 			FSparseVolumeRawSource SparseVolumeRawSource{};
-			SparseVolumeRawSource.PackedDataA = ImportOptions.PackedDataA;
-			SparseVolumeRawSource.PackedDataB = ImportOptions.PackedDataB;
-			SparseVolumeRawSource.SourceAssetFile = MoveTemp(LoadedFrameFile);
+
+			FOpenVDBToSVTConversionResult SVTResult;
+			SVTResult.Header = &SparseVolumeRawSource.Header;
+			SVTResult.PageTable = &SparseVolumeRawSource.PageTable;
+			SVTResult.PhysicalTileDataA = &SparseVolumeRawSource.PhysicalTileDataA;
+			SVTResult.PhysicalTileDataB = &SparseVolumeRawSource.PhysicalTileDataB;
+
+			const bool bConversionSuccess = ConvertOpenVDBToSparseVolumeTexture(
+				LoadedFrameFile,
+				ImportOptions,
+				&SVTResult,
+				false, FVector::Zero(), FVector::Zero());
+
+			if (!bConversionSuccess)
+			{
+				UE_LOG(LogSparseVolumeTextureFactory, Error, TEXT("Failed to convert OpenVDB file to SparseVolumeTexture: %s"), *FrameFilename);
+				return nullptr;
+			}
 
 			// Serialize the raw source data from this frame into the asset object.
 			{
@@ -595,6 +613,7 @@ UObject* USparseVolumeTextureFactory::FactoryCreateFile(UClass* InClass, UObject
 
 			if (ImportTask.ShouldCancel())
 			{
+				bOutOperationCanceled = true;
 				return nullptr;
 			}
 			ImportTask.EnterProgressFrame(1.0f, LOCTEXT("ConvertingVDBAnim", "Converting OpenVDB animation"));
diff --git a/Engine/Source/Editor/SparseVolumeTexture/Private/SparseVolumeTextureModule.cpp b/Engine/Source/Editor/SparseVolumeTexture/Private/SparseVolumeTextureModule.cpp
--- a/Engine/Source/Editor/SparseVolumeTexture/Private/SparseVolumeTextureModule.cpp
+++ b/Engine/Source/Editor/SparseVolumeTexture/Private/SparseVolumeTextureModule.cpp
@@ -3,8 +3,6 @@
 #include "SparseVolumeTextureModule.h"
 
 #include "SparseVolumeTextureOpenVDB.h"
-#include "SparseVolumeTexture/SparseVolumeTexture.h"
-#include "SparseVolumeTextureOpenVDBUtility.h"
 
 #define LOCTEXT_NAMESPACE "SparseVolumeTextureModule"
 
@@ -20,14 +18,10 @@ void FSparseVolumeTextureModule::StartupModule()
 	// Global registration of  the vdb types.
 	openvdb::initialize();
 #endif
-
-	// USparseVolumeTexture needs to access (only) this function for cooking.
-	OnConvertOpenVDBToSparseVolumeTexture().BindStatic(ConvertOpenVDBToSparseVolumeTexture);
 }
 
 void FSparseVolumeTextureModule::ShutdownModule()
 {
-	OnConvertOpenVDBToSparseVolumeTexture().Unbind();
 }
 
 #undef LOCTEXT_NAMESPACE
diff --git a/Engine/Source/Editor/SparseVolumeTexture/Private/SparseVolumeTextureOpenVDBUtility.cpp b/Engine/Source/Editor/SparseVolumeTexture/Private/SparseVolumeTextureOpenVDBUtility.cpp
--- a/Engine/Source/Editor/SparseVolumeTexture/Private/SparseVolumeTextureOpenVDBUtility.cpp
+++ b/Engine/Source/Editor/SparseVolumeTexture/Private/SparseVolumeTextureOpenVDBUtility.cpp
@@ -6,6 +6,7 @@
 #include "SparseVolumeTextureOpenVDB.h"
 #include "SparseVolumeTexture/SparseVolumeTexture.h"
 #include "OpenVDBGridAdapter.h"
+#include "OpenVDBImportOptions.h"
 
 DEFINE_LOG_CATEGORY_STATIC(LogSparseVolumeTextureOpenVDBUtility, Log, All);
 
@@ -157,11 +158,11 @@ bool GetOpenVDBGridInfo(TArray<uint8>& SourceFile, bool bCreateStrings, TArray<F
 	return false;
 }
 
-static EPixelFormat GetMultiComponentFormat(ESparseVolumePackedDataFormat Format, uint32 NumComponents)
+static EPixelFormat GetMultiComponentFormat(ESparseVolumeAttributesFormat Format, uint32 NumComponents)
 {
 	switch (Format)
 	{
-	case ESparseVolumePackedDataFormat::Unorm8:
+	case ESparseVolumeAttributesFormat::Unorm8:
 	{
 		switch (NumComponents)
 		{
@@ -172,7 +173,7 @@ static EPixelFormat GetMultiComponentFormat(ESparseVolumePackedDataFormat Format
 		}
 		break;
 	}
-	case ESparseVolumePackedDataFormat::Float16:
+	case ESparseVolumeAttributesFormat::Float16:
 	{
 		switch (NumComponents)
 		{
@@ -183,7 +184,7 @@ static EPixelFormat GetMultiComponentFormat(ESparseVolumePackedDataFormat Format
 		}
 		break;
 	}
-	case ESparseVolumePackedDataFormat::Float32:
+	case ESparseVolumeAttributesFormat::Float32:
 	{
 		switch (NumComponents)
 		{
@@ -218,8 +219,7 @@ static FIntVector3 UnpackPageTableEntry(uint32 Packed)
 
 bool ConvertOpenVDBToSparseVolumeTexture(
 	TArray<uint8>& SourceFile,
-	FSparseVolumeRawSourcePackedData& PackedDataA,
-	FSparseVolumeRawSourcePackedData& PackedDataB,
+	FOpenVDBImportOptions& ImportOptions,
 	FOpenVDBToSVTConversionResult* OutResult,
 	bool bOverrideActiveMinMax,
 	FVector ActiveMin,
@@ -227,46 +227,46 @@ bool ConvertOpenVDBToSparseVolumeTexture(
 {
 #if OPENVDB_AVAILABLE
 
-	constexpr uint32 NumPackedData = 2; // PackedDataA and PackedDataB, representing the two textures with voxel data
-	FSparseVolumeRawSourcePackedData* PackedData[NumPackedData] = { &PackedDataA, &PackedDataB };
+	constexpr uint32 NumAttributesDescs = 2; // AttributesA and AttributesB, representing the two textures with voxel data
+	TStaticArray<FOpenVDBSparseVolumeAttributesDesc, 2>& Attributes = ImportOptions.Attributes;
 
 	// Compute some basic info about the number of components and which format to use
-	uint32 NumActualComponents[NumPackedData] = {};
-	EPixelFormat MultiCompFormat[NumPackedData] = {};
-	uint32 FormatSize[NumPackedData] = {};
-	uint32 SingleComponentFormatSize[NumPackedData] = {};
-	bool bNormalizedFormat[NumPackedData] = {};
-	bool bHasValidSourceGrids[NumPackedData] = {};
+	uint32 NumActualComponents[NumAttributesDescs] = {};
+	EPixelFormat MultiCompFormat[NumAttributesDescs] = {};
+	uint32 FormatSize[NumAttributesDescs] = {};
+	uint32 SingleComponentFormatSize[NumAttributesDescs] = {};
+	bool bNormalizedFormat[NumAttributesDescs] = {};
+	bool bHasValidSourceGrids[NumAttributesDescs] = {};
 	bool bAnySourceGridIndicesValid = false;
 
-	for (uint32 PackedDataIdx = 0; PackedDataIdx < NumPackedData; ++PackedDataIdx)
+	for (uint32 AttributesIdx = 0; AttributesIdx < NumAttributesDescs; ++AttributesIdx)
 	{
 		uint32 NumRequiredComponents = 0;
 		for (uint32 ComponentIdx = 0; ComponentIdx < 4; ++ComponentIdx)
 		{
-			if (PackedData[PackedDataIdx]->SourceGridIndex[ComponentIdx] != INDEX_NONE)
+			if (Attributes[AttributesIdx].Mappings[ComponentIdx].SourceGridIndex != INDEX_NONE)
 			{
-				check(PackedData[PackedDataIdx]->SourceComponentIndex[ComponentIdx] != INDEX_NONE);
+				check(Attributes[AttributesIdx].Mappings[ComponentIdx].SourceComponentIndex != INDEX_NONE);
 				NumRequiredComponents = FMath::Max(ComponentIdx + 1, NumRequiredComponents);
-				bHasValidSourceGrids[PackedDataIdx] = true;
+				bHasValidSourceGrids[AttributesIdx] = true;
 				bAnySourceGridIndicesValid = true;
 			}
 		}
 
-		if (bHasValidSourceGrids[PackedDataIdx])
+		if (bHasValidSourceGrids[AttributesIdx])
 		{
-			NumActualComponents[PackedDataIdx] = NumRequiredComponents == 3 ? 4 : NumRequiredComponents; // We don't support formats with only 3 components
-			bNormalizedFormat[PackedDataIdx] = PackedData[PackedDataIdx]->Format == ESparseVolumePackedDataFormat::Unorm8;
-			MultiCompFormat[PackedDataIdx] = GetMultiComponentFormat(PackedData[PackedDataIdx]->Format, NumActualComponents[PackedDataIdx]);
+			NumActualComponents[AttributesIdx] = NumRequiredComponents == 3 ? 4 : NumRequiredComponents; // We don't support formats with only 3 components
+			bNormalizedFormat[AttributesIdx] = Attributes[AttributesIdx].Format == ESparseVolumeAttributesFormat::Unorm8;
+			MultiCompFormat[AttributesIdx] = GetMultiComponentFormat(Attributes[AttributesIdx].Format, NumActualComponents[AttributesIdx]);
 
-			if (MultiCompFormat[PackedDataIdx] == PF_Unknown)
+			if (MultiCompFormat[AttributesIdx] == PF_Unknown)
 			{
-				UE_LOG(LogSparseVolumeTextureOpenVDBUtility, Warning, TEXT("SparseVolumeTexture is set to use an unsupported format: %i"), (int32)PackedData[PackedDataIdx]->Format);
+				UE_LOG(LogSparseVolumeTextureOpenVDBUtility, Warning, TEXT("SparseVolumeTexture is set to use an unsupported format: %i"), (int32)Attributes[AttributesIdx].Format);
 				return false;
 			}
 
-			FormatSize[PackedDataIdx] = (uint32)GPixelFormats[(SIZE_T)MultiCompFormat[PackedDataIdx]].BlockBytes;
-			SingleComponentFormatSize[PackedDataIdx] = FormatSize[PackedDataIdx] / NumActualComponents[PackedDataIdx];
+			FormatSize[AttributesIdx] = (uint32)GPixelFormats[(SIZE_T)MultiCompFormat[AttributesIdx]].BlockBytes;
+			SingleComponentFormatSize[AttributesIdx] = FormatSize[AttributesIdx] / NumActualComponents[AttributesIdx];
 		}
 	}
 
@@ -285,22 +285,22 @@ bool ConvertOpenVDBToSparseVolumeTexture(
 	// Check that the source grid indices are valid
 	openvdb::GridPtrVecPtr Grids = Stream.getGrids();
 	const size_t NumSourceGrids = Grids ? Grids->size() : 0;
-	for (uint32 PackedDataIdx = 0; PackedDataIdx < NumPackedData; ++PackedDataIdx)
+	for (const FOpenVDBSparseVolumeAttributesDesc& AttributesDesc : ImportOptions.Attributes)
 	{
-		for (uint32 CompIdx = 0; CompIdx < 4; ++CompIdx)
+		for (const FOpenVDBSparseVolumeComponentMapping& Mapping : AttributesDesc.Mappings)
 		{
-			const uint32 SourceGridIndex = PackedData[PackedDataIdx]->SourceGridIndex[CompIdx];
-			if (SourceGridIndex != INDEX_NONE && SourceGridIndex >= NumSourceGrids)
+			const int32 SourceGridIndex = Mapping.SourceGridIndex;
+			if (SourceGridIndex != INDEX_NONE && SourceGridIndex >= (int32)NumSourceGrids)
 			{
-				UE_LOG(LogSparseVolumeTextureOpenVDBUtility, Warning, TEXT("SparseVolumeTexture has invalid index into the array of grids in the source file: %i"), (int32)SourceGridIndex);
+				UE_LOG(LogSparseVolumeTextureOpenVDBUtility, Warning, TEXT("SparseVolumeTexture has invalid index into the array of grids in the source file: %i"), SourceGridIndex);
 				return false;
 			}
 		}
 	}
 
 	FSparseVolumeAssetHeader& Header = *OutResult->Header;
-	Header.PackedDataAFormat = MultiCompFormat[0];
-	Header.PackedDataBFormat = MultiCompFormat[1];
+	Header.AttributesAFormat = MultiCompFormat[0];
+	Header.AttributesBFormat = MultiCompFormat[1];
 	Header.SourceVolumeResolution = FIntVector::ZeroValue;
 	
 	FIntVector SmallestAABBMin = FIntVector(INT32_MAX);
@@ -308,17 +308,17 @@ bool ConvertOpenVDBToSparseVolumeTexture(
 	// Compute per source grid data of up to 4 different grids (one per component)
 	TArray<TSharedPtr<IOpenVDBGridAdapterBase>> UniqueGridAdapters;
 	UniqueGridAdapters.SetNum((int32)NumSourceGrids);
-	TSharedPtr<IOpenVDBGridAdapterBase> GridAdapters[NumPackedData][4]{};
-	float GridBackgroundValues[NumPackedData][4]{};
-	float NormalizeScale[NumPackedData][4]{};
-	float NormalizeBias[NumPackedData][4]{};
-	for (uint32 PackedDataIdx = 0; PackedDataIdx < NumPackedData; ++PackedDataIdx)
+	TSharedPtr<IOpenVDBGridAdapterBase> GridAdapters[NumAttributesDescs][4]{};
+	float GridBackgroundValues[NumAttributesDescs][4]{};
+	float NormalizeScale[NumAttributesDescs][4]{};
+	float NormalizeBias[NumAttributesDescs][4]{};
+	for (uint32 AttributesIdx = 0; AttributesIdx < NumAttributesDescs; ++AttributesIdx)
 	{
 		for (uint32 CompIdx = 0; CompIdx < 4; ++CompIdx)
 		{
-			NormalizeScale[PackedDataIdx][CompIdx] = 1.0f;
-			const uint32 SourceGridIndex = PackedData[PackedDataIdx]->SourceGridIndex[CompIdx];
-			const uint32 SourceComponentIndex = PackedData[PackedDataIdx]->SourceComponentIndex[CompIdx];
+			NormalizeScale[AttributesIdx][CompIdx] = 1.0f;
+			const uint32 SourceGridIndex = Attributes[AttributesIdx].Mappings[CompIdx].SourceGridIndex;
+			const uint32 SourceComponentIndex = Attributes[AttributesIdx].Mappings[CompIdx].SourceComponentIndex;
 			if (SourceGridIndex == INDEX_NONE)
 			{
 				continue;
@@ -337,7 +337,7 @@ bool ConvertOpenVDBToSparseVolumeTexture(
 				}
 			}
 
-			GridAdapters[PackedDataIdx][CompIdx] = UniqueGridAdapters[SourceGridIndex];
+			GridAdapters[AttributesIdx][CompIdx] = UniqueGridAdapters[SourceGridIndex];
 
 			FOpenVDBGridInfo GridInfo = GetOpenVDBGridInfo(GridBase, 0, false);
 			if (!IsOpenVDBGridValid(GridInfo, TEXT("")))
@@ -359,15 +359,15 @@ bool ConvertOpenVDBToSparseVolumeTexture(
 			SmallestAABBMin.Y = FMath::Min(SmallestAABBMin.Y, GridInfo.VolumeActiveAABBMin.Y);
 			SmallestAABBMin.Z = FMath::Min(SmallestAABBMin.Z, GridInfo.VolumeActiveAABBMin.Z);
 
-			GridBackgroundValues[PackedDataIdx][CompIdx] = GridAdapters[PackedDataIdx][CompIdx]->GetBackgroundValue(SourceComponentIndex);
-			if (bNormalizedFormat[PackedDataIdx] && PackedData[PackedDataIdx]->bRemapInputForUnorm)
+			GridBackgroundValues[AttributesIdx][CompIdx] = GridAdapters[AttributesIdx][CompIdx]->GetBackgroundValue(SourceComponentIndex);
+			if (bNormalizedFormat[AttributesIdx] && Attributes[AttributesIdx].bRemapInputForUnorm)
 			{
 				float MinVal = 0.0f;
 				float MaxVal = 0.0f;
-				GridAdapters[PackedDataIdx][CompIdx]->GetMinMaxValue(SourceComponentIndex, &MinVal, &MaxVal);
+				GridAdapters[AttributesIdx][CompIdx]->GetMinMaxValue(SourceComponentIndex, &MinVal, &MaxVal);
 				const float Diff = MaxVal - MinVal;
-				NormalizeScale[PackedDataIdx][CompIdx] = MaxVal > SMALL_NUMBER ? (1.0f / Diff) : 1.0f;
-				NormalizeBias[PackedDataIdx][CompIdx] = -MinVal * NormalizeScale[PackedDataIdx][CompIdx];
+				NormalizeScale[AttributesIdx][CompIdx] = MaxVal > SMALL_NUMBER ? (1.0f / Diff) : 1.0f;
+				NormalizeBias[AttributesIdx][CompIdx] = -MinVal * NormalizeScale[AttributesIdx][CompIdx];
 			}
 		}
 	}
@@ -402,14 +402,14 @@ bool ConvertOpenVDBToSparseVolumeTexture(
 			{
 				// Check if the source grid component is used at all
 				bool bValueIsUsed = false;
-				for (uint32 PackedDataIdx = 0; PackedDataIdx < NumPackedData && !bValueIsUsed; ++PackedDataIdx)
+				for (uint32 AttributesIdx = 0; AttributesIdx < NumAttributesDescs && !bValueIsUsed; ++AttributesIdx)
 				{
-					for (uint32 DstCompIdx = 0; DstCompIdx < NumActualComponents[PackedDataIdx] && !bValueIsUsed; ++DstCompIdx)
+					for (uint32 DstCompIdx = 0; DstCompIdx < NumActualComponents[AttributesIdx] && !bValueIsUsed; ++DstCompIdx)
 					{
 						for (uint32 SrcCompIdx = 0; SrcCompIdx < NumComponents && !bValueIsUsed; ++SrcCompIdx)
 						{
-							const bool bIsBackgroundValue = (VoxelValues[SrcCompIdx] == GridBackgroundValues[PackedDataIdx][DstCompIdx]);
-							bValueIsUsed |= !bIsBackgroundValue && (PackedData[PackedDataIdx]->SourceGridIndex[DstCompIdx] == GridIdx) && (PackedData[PackedDataIdx]->SourceComponentIndex[DstCompIdx] == SrcCompIdx);
+							const bool bIsBackgroundValue = (VoxelValues[SrcCompIdx] == GridBackgroundValues[AttributesIdx][DstCompIdx]);
+							bValueIsUsed |= !bIsBackgroundValue && (Attributes[AttributesIdx].Mappings[DstCompIdx].SourceGridIndex == GridIdx) && (Attributes[AttributesIdx].Mappings[DstCompIdx].SourceComponentIndex == SrcCompIdx);
 						}
 					}
 				}
@@ -556,41 +556,41 @@ bool ConvertOpenVDBToSparseVolumeTexture(
 				const FIntVector3 TileLocalCoord = GridCoord % SPARSE_VOLUME_TILE_RES;
 
 				// Check all output components and splat VoxelValue if they map to this source grid/component
-				for (uint32 PackedDataIdx = 0; PackedDataIdx < NumPackedData; ++PackedDataIdx)
+				for (uint32 AttributesIdx = 0; AttributesIdx < NumAttributesDescs; ++AttributesIdx)
 				{
-					for (uint32 DstCompIdx = 0; DstCompIdx < NumActualComponents[PackedDataIdx]; ++DstCompIdx)
+					for (uint32 DstCompIdx = 0; DstCompIdx < NumActualComponents[AttributesIdx]; ++DstCompIdx)
 					{
 						for (uint32 SrcCompIdx = 0; SrcCompIdx < NumComponents; ++SrcCompIdx)
 						{
-							if ((PackedData[PackedDataIdx]->SourceGridIndex[DstCompIdx] != GridIdx) || (PackedData[PackedDataIdx]->SourceComponentIndex[DstCompIdx] != SrcCompIdx))
+							if ((Attributes[AttributesIdx].Mappings[DstCompIdx].SourceGridIndex != GridIdx) || (Attributes[AttributesIdx].Mappings[DstCompIdx].SourceComponentIndex != SrcCompIdx))
 							{
 								continue;
 							}
 
-							const float VoxelValueNormalized = FMath::Clamp(VoxelValues[SrcCompIdx] * NormalizeScale[PackedDataIdx][DstCompIdx] + NormalizeBias[PackedDataIdx][DstCompIdx], 0.0f, 1.0f);
+							const float VoxelValueNormalized = FMath::Clamp(VoxelValues[SrcCompIdx] * NormalizeScale[AttributesIdx][DstCompIdx] + NormalizeBias[AttributesIdx][DstCompIdx], 0.0f, 1.0f);
 
 							const size_t DstVoxelIndex =
 								(DstTileCoord.Z * SPARSE_VOLUME_TILE_RES + TileLocalCoord.Z) * Header.TileDataVolumeResolution.X * Header.TileDataVolumeResolution.Y +
 								(DstTileCoord.Y * SPARSE_VOLUME_TILE_RES + TileLocalCoord.Y) * Header.TileDataVolumeResolution.X +
 								(DstTileCoord.X * SPARSE_VOLUME_TILE_RES + TileLocalCoord.X);
-							const size_t DstCoord = DstVoxelIndex * FormatSize[PackedDataIdx] + DstCompIdx * SingleComponentFormatSize[PackedDataIdx];
+							const size_t DstCoord = DstVoxelIndex * FormatSize[AttributesIdx] + DstCompIdx * SingleComponentFormatSize[AttributesIdx];
 
-							switch (PackedData[PackedDataIdx]->Format)
+							switch (Attributes[AttributesIdx].Format)
 							{
-							case ESparseVolumePackedDataFormat::Unorm8:
+							case ESparseVolumeAttributesFormat::Unorm8:
 							{
-								PhysicalTileDataPtrs[PackedDataIdx][DstCoord] = uint8(VoxelValueNormalized * 255.0f);
+								PhysicalTileDataPtrs[AttributesIdx][DstCoord] = uint8(VoxelValueNormalized * 255.0f);
 								break;
 							}
-							case ESparseVolumePackedDataFormat::Float16:
+							case ESparseVolumeAttributesFormat::Float16:
 							{
 								const uint16 VoxelValue16FEncoded = FFloat16(VoxelValues[SrcCompIdx]).Encoded;
-								*((uint16*)(&PhysicalTileDataPtrs[PackedDataIdx][DstCoord])) = VoxelValue16FEncoded;
+								*((uint16*)(&PhysicalTileDataPtrs[AttributesIdx][DstCoord])) = VoxelValue16FEncoded;
 								break;
 							}
-							case ESparseVolumePackedDataFormat::Float32:
+							case ESparseVolumeAttributesFormat::Float32:
 							{
-								*((float*)(&PhysicalTileDataPtrs[PackedDataIdx][DstCoord])) = VoxelValues[SrcCompIdx];
+								*((float*)(&PhysicalTileDataPtrs[AttributesIdx][DstCoord])) = VoxelValues[SrcCompIdx];
 								break;
 							}
 							default:
diff --git a/Engine/Source/Editor/SparseVolumeTexture/Private/SparseVolumeTextureOpenVDBUtility.h b/Engine/Source/Editor/SparseVolumeTexture/Private/SparseVolumeTextureOpenVDBUtility.h
--- a/Engine/Source/Editor/SparseVolumeTexture/Private/SparseVolumeTextureOpenVDBUtility.h
+++ b/Engine/Source/Editor/SparseVolumeTexture/Private/SparseVolumeTextureOpenVDBUtility.h
@@ -35,14 +35,21 @@ struct FOpenVDBGridInfo
 	bool bHasUniformVoxels;
 };
 
+struct FOpenVDBToSVTConversionResult
+{
+	struct FSparseVolumeAssetHeader* Header;
+	TArray<uint32>* PageTable;
+	TArray<uint8>* PhysicalTileDataA;
+	TArray<uint8>* PhysicalTileDataB;
+};
+
 bool IsOpenVDBGridValid(const FOpenVDBGridInfo& GridInfo, const FString& Filename);
 
 bool GetOpenVDBGridInfo(TArray<uint8>& SourceFile, bool bCreateStrings, TArray<FOpenVDBGridInfo>* OutGridInfo);
 
 bool ConvertOpenVDBToSparseVolumeTexture(
 	TArray<uint8>& SourceFile,
-	struct FSparseVolumeRawSourcePackedData& PackedDataA,
-	struct FSparseVolumeRawSourcePackedData& PackedDataB,
+	struct FOpenVDBImportOptions& ImportOptions,
 	struct FOpenVDBToSVTConversionResult* OutResult,
 	bool bOverrideActiveMinMax,
 	FVector ActiveMin,
diff --git a/Engine/Source/Runtime/Engine/Classes/SparseVolumeTexture/SparseVolumeTexture.h b/Engine/Source/Runtime/Engine/Classes/SparseVolumeTexture/SparseVolumeTexture.h
--- a/Engine/Source/Runtime/Engine/Classes/SparseVolumeTexture/SparseVolumeTexture.h
+++ b/Engine/Source/Runtime/Engine/Classes/SparseVolumeTexture/SparseVolumeTexture.h
@@ -17,49 +17,41 @@ namespace UE { namespace Shader	{ enum class EValueType : uint8; } }
 
 #define SPARSE_VOLUME_TILE_RES	16
 
-struct ENGINE_API FOpenVDBToSVTConversionResult
+struct ENGINE_API FSparseVolumeAssetHeader
 {
-	struct FSparseVolumeAssetHeader* Header;
-	TArray<uint32>* PageTable;
-	TArray<uint8>* PhysicalTileDataA;
-	TArray<uint8>* PhysicalTileDataB;
-};
+	FIntVector3 PageTableVolumeResolution;
+	FIntVector3 TileDataVolumeResolution;
+	FIntVector3 SourceVolumeResolution;
+	FIntVector3 SourceVolumeAABBMin;
+	EPixelFormat AttributesAFormat;
+	EPixelFormat AttributesBFormat;
 
-DECLARE_DELEGATE_RetVal_SevenParams(bool, FConvertOpenVDBToSparseVolumeTextureDelegate, 
-	TArray<uint8>& SourceFile,
-	struct FSparseVolumeRawSourcePackedData& PackedDataA,
-	struct FSparseVolumeRawSourcePackedData& PackedDataB,
-	FOpenVDBToSVTConversionResult* OutResult,
-	bool bOverrideActiveMinMax,
-	FVector ActiveMin,
-	FVector ActiveMax);
+	// The current data format version for the header.
+	static const uint32 kVersion = 0;
 
-ENGINE_API FConvertOpenVDBToSparseVolumeTextureDelegate& OnConvertOpenVDBToSparseVolumeTexture();
+	// This version can be used to convert existing header to new version later.
+	uint32 Version;
 
-enum class ESparseVolumePackedDataFormat : uint8
-{
-	Unorm8 = 0,
-	Float16 = 1,
-	Float32 = 2,
-};
+	void Serialize(FArchive& Ar);
 
-// Describes how the grids of a source asset map to the components of a float4 of packed data in a SVT
-struct ENGINE_API FSparseVolumeRawSourcePackedData
-{
-	ESparseVolumePackedDataFormat Format;
-	FUintVector4 SourceGridIndex;
-	FUintVector4 SourceComponentIndex; // Index [0-3] of the component in the source grid.
-	bool bRemapInputForUnorm; // Maps the input from its minimum and maximum value into the [0-1] range. Clamps to [0-1] otherwise.
+	FSparseVolumeAssetHeader()
+		: PageTableVolumeResolution(FIntVector3(0, 0, 0))
+		, TileDataVolumeResolution(FIntVector3(0, 0, 0))
+		, SourceVolumeResolution(FIntVector3(0, 0, 0))
+		, AttributesAFormat(PF_Unknown)
+		, AttributesBFormat(PF_Unknown)
+		, Version(kVersion)
+	{
+	}
 };
 
-ENGINE_API FArchive& operator<<(FArchive& Ar, FSparseVolumeRawSourcePackedData& PackedData);
-
 // The structure represent the source asset in high quality. It is used to cook the runtime data
 struct ENGINE_API FSparseVolumeRawSource
 {
-	FSparseVolumeRawSourcePackedData PackedDataA;
-	FSparseVolumeRawSourcePackedData PackedDataB;
-	TArray<uint8> SourceAssetFile;
+	FSparseVolumeAssetHeader Header;
+	TArray<uint32> PageTable;
+	TArray<uint8> PhysicalTileDataA;
+	TArray<uint8> PhysicalTileDataB;
 
 	// The current data format version for the raw source data.
 	static const uint32 kVersion = 0;
@@ -75,33 +67,6 @@ struct ENGINE_API FSparseVolumeRawSource
 	}
 };
 
-struct ENGINE_API FSparseVolumeAssetHeader
-{
-	FIntVector3 PageTableVolumeResolution;
-	FIntVector3 TileDataVolumeResolution;
-	FIntVector3 SourceVolumeResolution;
-	FIntVector3 SourceVolumeAABBMin;
-	EPixelFormat PackedDataAFormat;
-	EPixelFormat PackedDataBFormat;
-
-	// The current data format version for the header.
-	static const uint32 kVersion = 0;
-
-	// This version can be used to convert existing header to new version later.
-	uint32 Version;
-
-	void Serialize(FArchive& Ar);
-
-	FSparseVolumeAssetHeader()
-		: PageTableVolumeResolution(FIntVector3(0, 0, 0))
-		, TileDataVolumeResolution(FIntVector3(0, 0, 0))
-		, SourceVolumeResolution(FIntVector3(0, 0, 0))
-		, PackedDataAFormat(PF_Unknown)
-		, Version(kVersion)
-	{
-	}
-};
-
 // The structure represent the runtime data cooked runtime data.
 struct ENGINE_API FSparseVolumeTextureRuntime
 {
diff --git a/Engine/Source/Runtime/Engine/Private/Materials/MaterialExpressions.cpp b/Engine/Source/Runtime/Engine/Private/Materials/MaterialExpressions.cpp
--- a/Engine/Source/Runtime/Engine/Private/Materials/MaterialExpressions.cpp
+++ b/Engine/Source/Runtime/Engine/Private/Materials/MaterialExpressions.cpp
@@ -25751,8 +25751,8 @@ UMaterialExpressionSparseVolumeTextureSample::UMaterialExpressionSparseVolumeTex
 
 #if WITH_EDITOR
 	Outputs.Reset();
-	Outputs.Add(FExpressionOutput(TEXT("PackedData A"), 1, 1, 1, 1, 1));
-	Outputs.Add(FExpressionOutput(TEXT("PackedData B"), 1, 1, 1, 1, 1));
+	Outputs.Add(FExpressionOutput(TEXT("Attributes A"), 1, 1, 1, 1, 1));
+	Outputs.Add(FExpressionOutput(TEXT("Attributes B"), 1, 1, 1, 1, 1));
 	bShowOutputNameOnPin = true;
 	bShowMaskColorsOnPin = false;
 #endif
@@ -25790,7 +25790,7 @@ uint32 UMaterialExpressionSparseVolumeTextureSample::GetOutputType(int32 OutputI
 	{
 	case 0:
 	case 1:
-		return MCT_Float;	// PackedData A and B defined in constructor for now
+		return MCT_Float;	// Attributes A and B defined in constructor for now
 		break;
 	}
 
@@ -25864,7 +25864,7 @@ int32 UMaterialExpressionSparseVolumeTextureSample::Compile(class FMaterialCompi
 		int32 CoordPageTable = Compiler->Mul(UVWPageTable, UniformCodeChunkIndex[ESparseVolumeTexture_PageTableSize]);
 
 		// Code chunks have the same ID when they map to the same string of code. PackedPhysicalTileCoord will automatically end up being the same code chunk used
-		// when calling Compile() for Output 0 (PackedDataA) and Output 1 (PackedDataB). This is perfect for this use case: 
+		// when calling Compile() for Output 0 (Attributes A) and Output 1 (Attributes B). This is perfect for this use case: 
 		// PackedPhysicalTileCoord will end up being evaluated only once for both outputs, so they both share a single page table lookup.
 		int32 PackedPhysicalTileCoord = Compiler->SparseVolumeTextureSample(PageTableTextureParamCodeChunkIndex, CoordPageTable);
 
diff --git a/Engine/Source/Runtime/Engine/Private/SparseVolumeTexture/SparseVolumeTexture.cpp b/Engine/Source/Runtime/Engine/Private/SparseVolumeTexture/SparseVolumeTexture.cpp
--- a/Engine/Source/Runtime/Engine/Private/SparseVolumeTexture/SparseVolumeTexture.cpp
+++ b/Engine/Source/Runtime/Engine/Private/SparseVolumeTexture/SparseVolumeTexture.cpp
@@ -26,13 +26,6 @@
 
 DEFINE_LOG_CATEGORY_STATIC(LogSparseVolumeTexture, Log, All);
 
-FConvertOpenVDBToSparseVolumeTextureDelegate ConvertOpenVDBToSparseVolumeTextureDelegate;
-
-FConvertOpenVDBToSparseVolumeTextureDelegate& OnConvertOpenVDBToSparseVolumeTexture()
-{
-	return ConvertOpenVDBToSparseVolumeTextureDelegate;
-}
-
 namespace UE
 {
 namespace SVT
@@ -54,27 +47,19 @@ namespace Private
 } // SVT
 } // UE
 
-FArchive& operator<<(FArchive& Ar, FSparseVolumeRawSourcePackedData& PackedData)
-{
-	UE::SVT::Private::SerializeEnumAs<uint8>(Ar, PackedData.Format);
-	Ar << PackedData.SourceGridIndex;
-	Ar << PackedData.SourceComponentIndex;
-	Ar << PackedData.bRemapInputForUnorm;
-
-	return Ar;
-}
-
 ////////////////////////////////////////////////////////////////////////////////////////////////
 
 void FSparseVolumeRawSource::Serialize(FArchive& Ar)
 {
+	Header.Serialize(Ar);
+
 	Ar << Version;
 
 	if (Version == 0)
 	{
-		Ar << PackedDataA;
-		Ar << PackedDataB;
-		Ar << SourceAssetFile;
+		Ar << PageTable;
+		Ar << PhysicalTileDataA;
+		Ar << PhysicalTileDataB;
 	}
 	else
 	{
@@ -95,8 +80,8 @@ void FSparseVolumeAssetHeader::Serialize(FArchive& Ar)
 		Ar << TileDataVolumeResolution;
 		Ar << SourceVolumeResolution;
 		Ar << SourceVolumeAABBMin;
-		UE::SVT::Private::SerializeEnumAs<uint8>(Ar, PackedDataAFormat);
-		UE::SVT::Private::SerializeEnumAs<uint8>(Ar, PackedDataBFormat);
+		UE::SVT::Private::SerializeEnumAs<uint8>(Ar, AttributesAFormat);
+		UE::SVT::Private::SerializeEnumAs<uint8>(Ar, AttributesBFormat);
 	}
 	else
 	{
@@ -175,8 +160,8 @@ void FSparseVolumeTextureSceneProxy::InitRHI()
 	// Tile data
 	{
 		FIntVector3 TileDataVolumeResolution = SparseVolumeTextureRuntime->Header.TileDataVolumeResolution;
-		EPixelFormat VoxelFormatA = SparseVolumeTextureRuntime->Header.PackedDataAFormat;
-		EPixelFormat VoxelFormatB = SparseVolumeTextureRuntime->Header.PackedDataBFormat;
+		EPixelFormat VoxelFormatA = SparseVolumeTextureRuntime->Header.AttributesAFormat;
+		EPixelFormat VoxelFormatB = SparseVolumeTextureRuntime->Header.AttributesBFormat;
 		const FUpdateTextureRegion3D UpdateRegion(0, 0, 0, 0, 0, 0, TileDataVolumeResolution.X, TileDataVolumeResolution.Y, TileDataVolumeResolution.Z);
 
 		// A
@@ -240,26 +225,11 @@ bool FSparseVolumeTextureFrame::BuildRuntimeData()
 		FSparseVolumeRawSource SparseVolumeRawSource;
 		SparseVolumeRawSource.Serialize(RawDataArchiveReader);
 
-		// Then, convert the raw source data (OpenVDB) to SVT
-		FOpenVDBToSVTConversionResult SVTResult;
-		SVTResult.Header = &SparseVolumeTextureRuntime.Header;
-		SVTResult.PageTable = &SparseVolumeTextureRuntime.PageTable;
-		SVTResult.PhysicalTileDataA = &SparseVolumeTextureRuntime.PhysicalTileDataA;
-		SVTResult.PhysicalTileDataB = &SparseVolumeTextureRuntime.PhysicalTileDataB;
-
-		const bool bSuccess = ConvertOpenVDBToSparseVolumeTextureDelegate.IsBound() && ConvertOpenVDBToSparseVolumeTextureDelegate.Execute(
-			SparseVolumeRawSource.SourceAssetFile,
-			SparseVolumeRawSource.PackedDataA,
-			SparseVolumeRawSource.PackedDataB,
-			&SVTResult,
-			false, FVector::Zero(), FVector::Zero());
-		
-		if (!bSuccess)
-		{
-			// Clear any writes that may have happened to this data during the cook attempt
-			SparseVolumeTextureRuntime = {};
-			return false;
-		}
+		// Then, convert the raw source data to SVT
+		SparseVolumeTextureRuntime.Header = SparseVolumeRawSource.Header;
+		SparseVolumeTextureRuntime.PageTable = MoveTemp(SparseVolumeRawSource.PageTable);
+		SparseVolumeTextureRuntime.PhysicalTileDataA = MoveTemp(SparseVolumeRawSource.PhysicalTileDataA);
+		SparseVolumeTextureRuntime.PhysicalTileDataB = MoveTemp(SparseVolumeRawSource.PhysicalTileDataB);
 
 		// Now unload the raw data
 		RawData.UnloadData();
