Import of .vdb sequences for SparseVolumeTextures now uses multithreading to parallelize per-frame processing. Fixed background/fallback values always being zero, even if the value specified in the source .vdb file is non-zero.

#rnx

diff --git a/Engine/Source/Editor/SparseVolumeTexture/Private/SparseVolumeTextureFactory.cpp b/Engine/Source/Editor/SparseVolumeTexture/Private/SparseVolumeTextureFactory.cpp
--- a/Engine/Source/Editor/SparseVolumeTexture/Private/SparseVolumeTextureFactory.cpp
+++ b/Engine/Source/Editor/SparseVolumeTexture/Private/SparseVolumeTextureFactory.cpp
@@ -14,13 +14,19 @@
 #include "Misc/Paths.h"
 #include "Misc/ScopedSlowTask.h"
 #include "Misc/FileHelper.h"
+#include "Async/Async.h"
 
 #include "Editor.h"
 
 #include "OpenVDBImportWindow.h"
 #include "HAL/PlatformApplicationMisc.h"
+#include "HAL/Event.h"
+#include "HAL/PlatformProcess.h"
 #include "Interfaces/IMainFrameModule.h"
 
+#include <atomic>
+#include <mutex>
+
 #define LOCTEXT_NAMESPACE "USparseVolumeTextureFactory"
 
 DEFINE_LOG_CATEGORY_STATIC(LogSparseVolumeTextureFactory, Log, All);
@@ -537,88 +543,170 @@ UObject* USparseVolumeTextureFactory::FactoryCreateFile(UClass* InClass, UObject
 		// Allocate space for each frame
 		AnimatedSVTexture->Frames.SetNum(NumFrames);
 
+		FEvent* AllTasksFinishedEvent = FPlatformProcess::GetSynchEventFromPool();
+		std::atomic_bool bErrored = false;
+		std::atomic_bool bCanceled = false;
+		std::atomic_int FinishedTasksCounter = 0; // Will be incremented even if frame processing failed
+		std::atomic_int ProcessedFramesCounter = 0;
+		std::mutex Mutex;
 		FBox VolumeBounds(FVector(FLT_MAX), FVector(-FLT_MAX));
 
 		// Load individual frames, check them for compatibility with the first loaded file and append them to the resulting asset
 		for (int32 FrameIdx = 0; FrameIdx < NumFrames; ++FrameIdx)
 		{
-			const FString& FrameFilename = PreviewData.SequenceFilenames[FrameIdx];
+			// Increments the atomic counter when going out of scope. Triggers an event once the counter reaches a given value.
+			struct FScopedIncrementer
+			{
+				std::atomic_int& Counter;
+				int32 MaxValue;
+				FEvent* Event;
+				explicit FScopedIncrementer(std::atomic_int& InCounter, int32 InMaxValue, FEvent* InEvent) 
+					: Counter(InCounter), MaxValue(InMaxValue), Event(InEvent) {}
+				~FScopedIncrementer() 
+				{ 
+					if ((Counter.fetch_add(1) + 1) == MaxValue)
+					{
+						Event->Trigger();
+					}
+				}
+			};
 
-			UE_LOG(LogSparseVolumeTextureFactory, Display, TEXT("Loading OpenVDB sequence frame #%i %s."), FrameIdx, *FrameFilename);
+			AsyncTask(ENamedThreads::AnyNormalThreadNormalTask, 
+				[FrameIdx, NumFrames, &PreviewData, &ImportOptions, &AnimatedSVTexture, &ExpandVolumeBounds,
+				AllTasksFinishedEvent, &bErrored, &bCanceled, &Mutex, &FinishedTasksCounter, &ProcessedFramesCounter, &VolumeBounds]()
+				{
+					// Ensure the FinishedTasksCounter will be incremented in all cases
+					FScopedIncrementer Incremeter(FinishedTasksCounter, NumFrames, AllTasksFinishedEvent);
 
-			// Load file and get info about each contained grid
-			TArray<uint8> LoadedFrameFile;
-			if (!FFileHelper::LoadFileToArray(LoadedFrameFile, *FrameFilename))
-			{
-				UE_LOG(LogSparseVolumeTextureFactory, Error, TEXT("OpenVDB file could not be loaded: %s"), *FrameFilename);
-				return nullptr;
-			}
+					if (bErrored.load() || bCanceled.load())
+					{
+						return;
+					}
 
-			TArray<FOpenVDBGridInfo> FrameGridInfo;
-			if (!GetOpenVDBGridInfo(LoadedFrameFile, true, &FrameGridInfo))
-			{
-				UE_LOG(LogSparseVolumeTextureFactory, Error, TEXT("Failed to read OpenVDB file: %s"), *FrameFilename);
-				return nullptr;
-			}
+					const FString& FrameFilename = PreviewData.SequenceFilenames[FrameIdx];
 
-			// Sanity check for compatibility
-			for (const FOpenVDBSparseVolumeAttributesDesc& AttributesDesc : ImportOptions.Attributes)
-			{
-				for (const FOpenVDBSparseVolumeComponentMapping& Mapping : AttributesDesc.Mappings)
-				{
-					const uint32 SourceGridIndex = Mapping.SourceGridIndex;
-					if (SourceGridIndex != INDEX_NONE)
+					UE_LOG(LogSparseVolumeTextureFactory, Display, TEXT("Loading OpenVDB sequence frame #%i %s."), FrameIdx, *FrameFilename);
+
+					// Load file and get info about each contained grid
+					TArray<uint8> LoadedFrameFile;
+					if (!FFileHelper::LoadFileToArray(LoadedFrameFile, *FrameFilename))
 					{
-						if ((int32)SourceGridIndex >= FrameGridInfo.Num())
-						{
-							UE_LOG(LogSparseVolumeTextureFactory, Error, TEXT("OpenVDB file is incompatible with other frames in the sequence: %s"), *FrameFilename);
-							return nullptr;
-						}
-						const auto& OrigSourceGrid = PreviewData.GridInfo[SourceGridIndex];
-						const auto& FrameSourceGrid = FrameGridInfo[SourceGridIndex];
-						if (OrigSourceGrid.Type != FrameSourceGrid.Type || OrigSourceGrid.Name != FrameSourceGrid.Name)
+						UE_LOG(LogSparseVolumeTextureFactory, Error, TEXT("OpenVDB file could not be loaded: %s"), *FrameFilename);
+						bErrored.store(true);
+						return;
+					}
+
+					TArray<FOpenVDBGridInfo> FrameGridInfo;
+					if (!GetOpenVDBGridInfo(LoadedFrameFile, true, &FrameGridInfo))
+					{
+						UE_LOG(LogSparseVolumeTextureFactory, Error, TEXT("Failed to read OpenVDB file: %s"), *FrameFilename);
+						bErrored.store(true);
+						return;
+					}
+
+					// Sanity check for compatibility
+					for (const FOpenVDBSparseVolumeAttributesDesc& AttributesDesc : ImportOptions.Attributes)
+					{
+						for (const FOpenVDBSparseVolumeComponentMapping& Mapping : AttributesDesc.Mappings)
 						{
-							UE_LOG(LogSparseVolumeTextureFactory, Error, TEXT("OpenVDB file is incompatible with other frames in the sequence: %s"), *FrameFilename);
-							return nullptr;
+							const uint32 SourceGridIndex = Mapping.SourceGridIndex;
+							if (SourceGridIndex != INDEX_NONE)
+							{
+								if ((int32)SourceGridIndex >= FrameGridInfo.Num())
+								{
+									UE_LOG(LogSparseVolumeTextureFactory, Error, TEXT("OpenVDB file is incompatible with other frames in the sequence: %s"), *FrameFilename);
+									bErrored.store(true);
+									return;
+								}
+								const FOpenVDBGridInfo& OrigSourceGrid = PreviewData.GridInfo[SourceGridIndex];
+								const FOpenVDBGridInfo& FrameSourceGrid = FrameGridInfo[SourceGridIndex];
+								if (OrigSourceGrid.Type != FrameSourceGrid.Type || OrigSourceGrid.Name != FrameSourceGrid.Name)
+								{
+									UE_LOG(LogSparseVolumeTextureFactory, Error, TEXT("OpenVDB file is incompatible with other frames in the sequence: %s"), *FrameFilename);
+									bErrored.store(true);
+									return;
+								}
+							}
 						}
 					}
-				}
-			}
 
-			ExpandVolumeBounds(ImportOptions, FrameGridInfo, VolumeBounds);
+					FSparseVolumeRawSource SparseVolumeRawSource{};
+
+					FOpenVDBToSVTConversionResult SVTResult;
+					SVTResult.Header = &SparseVolumeRawSource.Header;
+					SVTResult.PageTable = &SparseVolumeRawSource.PageTable;
+					SVTResult.PhysicalTileDataA = &SparseVolumeRawSource.PhysicalTileDataA;
+					SVTResult.PhysicalTileDataB = &SparseVolumeRawSource.PhysicalTileDataB;
+
+					const bool bConversionSuccess = ConvertOpenVDBToSparseVolumeTexture(
+						LoadedFrameFile,
+						ImportOptions,
+						&SVTResult,
+						false, FVector::Zero(), FVector::Zero());
 
-			FSparseVolumeRawSource SparseVolumeRawSource{};
+					if (!bConversionSuccess)
+					{
+						UE_LOG(LogSparseVolumeTextureFactory, Error, TEXT("Failed to convert OpenVDB file to SparseVolumeTexture: %s"), *FrameFilename);
+						bErrored.store(true);
+						return;
+					}
 
-			FOpenVDBToSVTConversionResult SVTResult;
-			SVTResult.Header = &SparseVolumeRawSource.Header;
-			SVTResult.PageTable = &SparseVolumeRawSource.PageTable;
-			SVTResult.PhysicalTileDataA = &SparseVolumeRawSource.PhysicalTileDataA;
-			SVTResult.PhysicalTileDataB = &SparseVolumeRawSource.PhysicalTileDataB;
+					// Serialize the raw source data from this frame into the asset object.
+					{
+						UE::Serialization::FEditorBulkDataWriter RawDataArchiveWriter(AnimatedSVTexture->Frames[FrameIdx].RawData);
+						SparseVolumeRawSource.Serialize(RawDataArchiveWriter);
+					}
 
-			const bool bConversionSuccess = ConvertOpenVDBToSparseVolumeTexture(
-				LoadedFrameFile,
-				ImportOptions,
-				&SVTResult,
-				false, FVector::Zero(), FVector::Zero());
+					// Update sequence volume bounds and increment ProcessedFramesCounter
+					{
+						std::lock_guard<std::mutex> Lock(Mutex);
+						ExpandVolumeBounds(ImportOptions, FrameGridInfo, VolumeBounds);
+					}
+					ProcessedFramesCounter.fetch_add(1);
+				});
+		}
 
-			if (!bConversionSuccess)
+		// Wait for frames to be processed
+		{
+			int NumFinishedTasks = 0;
+			int NumProcessedFrames = 0;
+			
+			while (NumFinishedTasks < NumFrames)
 			{
-				UE_LOG(LogSparseVolumeTextureFactory, Error, TEXT("Failed to convert OpenVDB file to SparseVolumeTexture: %s"), *FrameFilename);
-				return nullptr;
-			}
+				// We can't block here because we want to regularly update the progress bar and check for user input.
+				const uint32 WaitTimeMS = 2;
+				AllTasksFinishedEvent->Wait(WaitTimeMS);
 
-			// Serialize the raw source data from this frame into the asset object.
-			{
-				UE::Serialization::FEditorBulkDataWriter RawDataArchiveWriter(AnimatedSVTexture->Frames[FrameIdx].RawData);
-				SparseVolumeRawSource.Serialize(RawDataArchiveWriter);
-			}
+				if (!bCanceled.load() && !bErrored.load() && ImportTask.ShouldCancel())
+				{
+					bCanceled.store(true);
+				}
 
-			if (ImportTask.ShouldCancel())
-			{
-				bOutOperationCanceled = true;
-				return nullptr;
+				const int NewNumFinishedTasks = FinishedTasksCounter.load();
+				if (NewNumFinishedTasks > NumFinishedTasks)
+				{
+					const int NewNumProcessedFrames = ProcessedFramesCounter.load();
+					if (NewNumProcessedFrames > NumProcessedFrames && !bErrored.load())
+					{
+						const float Progress = float(NewNumProcessedFrames - NumProcessedFrames);
+						ImportTask.EnterProgressFrame(Progress, LOCTEXT("ConvertingVDBAnim", "Converting OpenVDB animation"));
+					}
+
+					NumFinishedTasks = NewNumFinishedTasks;
+					NumProcessedFrames = NewNumProcessedFrames;
+				}
 			}
-			ImportTask.EnterProgressFrame(1.0f, LOCTEXT("ConvertingVDBAnim", "Converting OpenVDB animation"));
+		}
+		FPlatformProcess::ReturnSynchEventToPool(AllTasksFinishedEvent);
+
+		if (bCanceled.load())
+		{
+			bOutOperationCanceled = true;
+			return nullptr;
+		}
+		if (bErrored.load())
+		{
+			return nullptr;
 		}
 
 		AnimatedSVTexture->VolumeBounds = VolumeBounds;
diff --git a/Engine/Source/Editor/SparseVolumeTexture/Private/SparseVolumeTextureOpenVDBUtility.cpp b/Engine/Source/Editor/SparseVolumeTexture/Private/SparseVolumeTextureOpenVDBUtility.cpp
--- a/Engine/Source/Editor/SparseVolumeTexture/Private/SparseVolumeTextureOpenVDBUtility.cpp
+++ b/Engine/Source/Editor/SparseVolumeTexture/Private/SparseVolumeTextureOpenVDBUtility.cpp
@@ -528,6 +528,63 @@ bool ConvertOpenVDBToSparseVolumeTexture(
 		}
 	}
 
+	auto WriteVoxelComponent = [&](const FIntVector3& DstCoord, uint32 AttributesIdx, uint32 CompIdx, float Value)
+	{
+		const size_t DstVoxelIndex = DstCoord.Z * Header.TileDataVolumeResolution.X * Header.TileDataVolumeResolution.Y 
+			+ DstCoord.Y * Header.TileDataVolumeResolution.X 
+			+ DstCoord.X;
+
+		const float VoxelValue = Value;
+		const float VoxelValueNormalized = FMath::Clamp(VoxelValue * NormalizeScale[AttributesIdx][CompIdx] + NormalizeBias[AttributesIdx][CompIdx], 0.0f, 1.0f);
+
+		const size_t DstIdx = DstVoxelIndex * FormatSize[AttributesIdx] + CompIdx * SingleComponentFormatSize[AttributesIdx];
+
+		switch (Attributes[AttributesIdx].Format)
+		{
+		case ESparseVolumeAttributesFormat::Unorm8:
+		{
+			PhysicalTileDataPtrs[AttributesIdx][DstIdx] = uint8(VoxelValueNormalized * 255.0f);
+			break;
+		}
+		case ESparseVolumeAttributesFormat::Float16:
+		{
+			const uint16 VoxelValue16FEncoded = FFloat16(VoxelValue).Encoded;
+			*((uint16*)(&PhysicalTileDataPtrs[AttributesIdx][DstIdx])) = VoxelValue16FEncoded;
+			break;
+		}
+		case ESparseVolumeAttributesFormat::Float32:
+		{
+			*((float*)(&PhysicalTileDataPtrs[AttributesIdx][DstIdx])) = VoxelValue;
+			break;
+		}
+		default:
+			checkNoEntry();
+			break;
+		}
+	};
+
+	// Fill null tile with fallback/background value
+	if (bAnyEmptyPageExists)
+	{
+		for (uint32 AttributesIdx = 0; AttributesIdx < NumAttributesDescs; ++AttributesIdx)
+		{
+			for (uint32 Z = 0; Z < SPARSE_VOLUME_TILE_RES; ++Z)
+			{
+				for (uint32 Y = 0; Y < SPARSE_VOLUME_TILE_RES; ++Y)
+				{
+					for (uint32 X = 0; X < SPARSE_VOLUME_TILE_RES; ++X)
+					{
+						for (uint32 CompIdx = 0; CompIdx < NumActualComponents[AttributesIdx]; ++CompIdx)
+						{
+							WriteVoxelComponent(FIntVector3(X, Y, Z), AttributesIdx, CompIdx, GridBackgroundValues[AttributesIdx][CompIdx]);
+						}
+					}
+				}
+			}
+		}
+	}
+
+	// Write physical tile data
 	for (uint32 GridIdx = 0; GridIdx < NumSourceGrids; ++GridIdx)
 	{
 		if (!UniqueGridAdapters[GridIdx])
@@ -556,47 +613,19 @@ bool ConvertOpenVDBToSparseVolumeTexture(
 				const FIntVector3 TileLocalCoord = GridCoord % SPARSE_VOLUME_TILE_RES;
 
 				// Check all output components and splat VoxelValue if they map to this source grid/component
-				for (uint32 AttributesIdx = 0; AttributesIdx < NumAttributesDescs; ++AttributesIdx)
+				for (uint32 SrcCompIdx = 0; SrcCompIdx < NumComponents; ++SrcCompIdx)
 				{
-					for (uint32 DstCompIdx = 0; DstCompIdx < NumActualComponents[AttributesIdx]; ++DstCompIdx)
+					for (uint32 AttributesIdx = 0; AttributesIdx < NumAttributesDescs; ++AttributesIdx)
 					{
-						for (uint32 SrcCompIdx = 0; SrcCompIdx < NumComponents; ++SrcCompIdx)
+						for (uint32 DstCompIdx = 0; DstCompIdx < NumActualComponents[AttributesIdx]; ++DstCompIdx)
 						{
 							if ((Attributes[AttributesIdx].Mappings[DstCompIdx].SourceGridIndex != GridIdx) || (Attributes[AttributesIdx].Mappings[DstCompIdx].SourceComponentIndex != SrcCompIdx))
 							{
 								continue;
 							}
 
-							const float VoxelValueNormalized = FMath::Clamp(VoxelValues[SrcCompIdx] * NormalizeScale[AttributesIdx][DstCompIdx] + NormalizeBias[AttributesIdx][DstCompIdx], 0.0f, 1.0f);
-
-							const size_t DstVoxelIndex =
-								(DstTileCoord.Z * SPARSE_VOLUME_TILE_RES + TileLocalCoord.Z) * Header.TileDataVolumeResolution.X * Header.TileDataVolumeResolution.Y +
-								(DstTileCoord.Y * SPARSE_VOLUME_TILE_RES + TileLocalCoord.Y) * Header.TileDataVolumeResolution.X +
-								(DstTileCoord.X * SPARSE_VOLUME_TILE_RES + TileLocalCoord.X);
-							const size_t DstCoord = DstVoxelIndex * FormatSize[AttributesIdx] + DstCompIdx * SingleComponentFormatSize[AttributesIdx];
-
-							switch (Attributes[AttributesIdx].Format)
-							{
-							case ESparseVolumeAttributesFormat::Unorm8:
-							{
-								PhysicalTileDataPtrs[AttributesIdx][DstCoord] = uint8(VoxelValueNormalized * 255.0f);
-								break;
-							}
-							case ESparseVolumeAttributesFormat::Float16:
-							{
-								const uint16 VoxelValue16FEncoded = FFloat16(VoxelValues[SrcCompIdx]).Encoded;
-								*((uint16*)(&PhysicalTileDataPtrs[AttributesIdx][DstCoord])) = VoxelValue16FEncoded;
-								break;
-							}
-							case ESparseVolumeAttributesFormat::Float32:
-							{
-								*((float*)(&PhysicalTileDataPtrs[AttributesIdx][DstCoord])) = VoxelValues[SrcCompIdx];
-								break;
-							}
-							default:
-								checkNoEntry();
-								break;
-							}
+							const FIntVector3 DstCoord = DstTileCoord * SPARSE_VOLUME_TILE_RES + TileLocalCoord;
+							WriteVoxelComponent(DstCoord, AttributesIdx, DstCompIdx, VoxelValues[SrcCompIdx]);
 						}
 					}
 				}
