Fixed OpenVDB import being limited to a maximum source asset file size of 2GB.

#rnx

diff --git a/Engine/Source/Editor/SparseVolumeTexture/Private/SparseVolumeTextureFactory.cpp b/Engine/Source/Editor/SparseVolumeTexture/Private/SparseVolumeTextureFactory.cpp
--- a/Engine/Source/Editor/SparseVolumeTexture/Private/SparseVolumeTextureFactory.cpp
+++ b/Engine/Source/Editor/SparseVolumeTexture/Private/SparseVolumeTextureFactory.cpp
@@ -159,7 +159,7 @@ static TArray<FString> FindOpenVDBSequenceFileNames(const FString& Filename)
 
 struct FOpenVDBPreviewData
 {
-	TArray<uint8> LoadedFile;
+	TArray64<uint8> LoadedFile;
 	TArray<FOpenVDBGridInfo> GridInfo;
 	TArray<TSharedPtr<FOpenVDBGridInfo>> GridInfoPtrs;
 	TArray<TSharedPtr<FOpenVDBGridComponentInfo>> GridComponentInfoPtrs;
@@ -570,7 +570,7 @@ UObject* USparseVolumeTextureFactory::FactoryCreateFile(UClass* InClass, UObject
 
 				// Load file and get info about each contained grid
 				const FString& FrameFilename = PreviewData.SequenceFilenames[FrameIdx];
-				TArray<uint8> LoadedFrameFile;
+				TArray64<uint8> LoadedFrameFile;
 				if (!FFileHelper::LoadFileToArray(LoadedFrameFile, *FrameFilename))
 				{
 					UE_LOG(LogSparseVolumeTextureFactory, Error, TEXT("OpenVDB file could not be loaded: %s"), *FrameFilename);
@@ -669,7 +669,7 @@ UObject* USparseVolumeTextureFactory::FactoryCreateFile(UClass* InClass, UObject
 					UE_LOG(LogSparseVolumeTextureFactory, Display, TEXT("Loading OpenVDB sequence frame #%i %s."), FrameIdx, *FrameFilename);
 
 					// Load file
-					TArray<uint8> LoadedFrameFile;
+					TArray64<uint8> LoadedFrameFile;
 					if (!FFileHelper::LoadFileToArray(LoadedFrameFile, *FrameFilename))
 					{
 						UE_LOG(LogSparseVolumeTextureFactory, Error, TEXT("OpenVDB file could not be loaded: %s"), *FrameFilename);
diff --git a/Engine/Source/Editor/SparseVolumeTexture/Private/SparseVolumeTextureOpenVDBUtility.cpp b/Engine/Source/Editor/SparseVolumeTexture/Private/SparseVolumeTextureOpenVDBUtility.cpp
--- a/Engine/Source/Editor/SparseVolumeTexture/Private/SparseVolumeTextureOpenVDBUtility.cpp
+++ b/Engine/Source/Editor/SparseVolumeTexture/Private/SparseVolumeTextureOpenVDBUtility.cpp
@@ -14,15 +14,75 @@ DEFINE_LOG_CATEGORY_STATIC(LogSparseVolumeTextureOpenVDBUtility, Log, All);
 
 namespace
 {
-	// Utility class acting as adapter between TArray<uint8> and std::istream
+	// Utility class acting as adapter between TArray64<uint8> and std::istream.
+	// In order to work around a problem where a std::streambuf implementation is limited to
+	// 2GB buffers, we need to manually update the pointers whenever we are done processing a 2GB chunk.
 	class FArrayUint8StreamBuf : public std::streambuf
 	{
 	public:
-		explicit FArrayUint8StreamBuf(TArray<uint8>& Array)
+		explicit FArrayUint8StreamBuf(TArray64<uint8>& Array)
 		{
 			char* Data = (char*)Array.GetData();
-			setg(Data, Data, Data + Array.Num());
+			size_t Num = Array.Num();
+			FileBegin = Data;
+			FileEnd = Data + Num;
+			FileRead = Data;
 		}
+
+		// Calls setg() to with a set of pointers exposing a window into the input file.
+		// Returns true if there are any more bytes to be read.
+		bool UpdatePointers()
+		{
+			StreamBegin = FileRead;
+			StreamEnd = FMath::Min(FileEnd, StreamBegin + ChunkSize);
+			StreamRead = StreamBegin;
+			FileRead = StreamRead;
+			setg(StreamBegin, StreamRead, StreamEnd);
+			const bool bHasBytesToRead = StreamBegin < StreamEnd;
+			return bHasBytesToRead;
+		}
+
+		// This function is called by the parent class and is expected to be implemented by subclasses.
+		// It requests n bytes to be copied into s. 
+		std::streamsize xsgetn(char* s, std::streamsize n) override
+		{
+			for (std::streamsize ReadBytes = 0; ReadBytes < n;)
+			{
+				// We read all bytes of the input file. gptr() is the current read ptr and egptr() is the end ptr.
+				if (gptr() == egptr() && !UpdatePointers())
+				{
+					check(gptr() == FileEnd);
+					return ReadBytes;
+				}
+				// Try to read n bytes but make a smaller read if we would read past the current ptr window exposed to streambuf.
+				const std::streamsize Available = FMath::Min(n - ReadBytes, static_cast<std::streamsize>(egptr() - gptr()));
+				memcpy(s, gptr(), Available);
+				// Advance pointers and the ReadBytes counter
+				s += Available;
+				ReadBytes += Available;
+				StreamRead = gptr() + Available;
+				FileRead = StreamRead;
+				// Update the Next ptr in streambuf
+				setg(StreamBegin, StreamRead, StreamEnd);
+			}
+			return n;
+		}
+
+		// Get the current character but don't advance the position. This is called by uflow() in the parent class when it runs out of bytes.
+		// We use it to move the exposed window into the file data.
+		int_type underflow() override
+		{
+			return (gptr() == egptr() && !UpdatePointers()) ? traits_type::eof() : *gptr();
+		}
+
+	private:
+		static constexpr size_t ChunkSize = INT32_MAX;	// The size of the range to expose to std::streambuf with setg()
+		char* FileBegin = nullptr;						// Begin ptr of the file data
+		char* FileEnd = nullptr;						// One byte past the end of the file data
+		char* FileRead = nullptr;						// The position up to which the file data has been exposed to/processed by the streambuf.
+		char* StreamBegin = nullptr;					// The begin ptr of the currently exposed file chunk.
+		char* StreamEnd = nullptr;						// The end ptr of the currently exposed file chunk.
+		char* StreamRead = nullptr;						// The last value of the read/next ptr set with setg().
 	};
 }
 
@@ -130,7 +190,7 @@ bool IsOpenVDBGridValid(const FOpenVDBGridInfo& GridInfo, const FString& Filenam
 	return true;
 }
 
-bool GetOpenVDBGridInfo(TArray<uint8>& SourceFile, bool bCreateStrings, TArray<FOpenVDBGridInfo>* OutGridInfo)
+bool GetOpenVDBGridInfo(TArray64<uint8>& SourceFile, bool bCreateStrings, TArray<FOpenVDBGridInfo>* OutGridInfo)
 {
 #if OPENVDB_AVAILABLE
 	FArrayUint8StreamBuf StreamBuf(SourceFile);
@@ -202,7 +262,7 @@ class FSparseVolumeRawSourceConstructionOpenVDBAdapter : public ISparseVolumeRaw
 {
 public:
 
-	bool Initialize(TArray<uint8>& SourceFile, const FOpenVDBImportOptions& ImportOptions, const FIntVector3& InVolumeBoundsMin)
+	bool Initialize(TArray64<uint8>& SourceFile, const FOpenVDBImportOptions& ImportOptions, const FIntVector3& InVolumeBoundsMin)
 	{
 		Attributes = ImportOptions.Attributes;
 		VolumeBoundsMin = InVolumeBoundsMin;
@@ -408,7 +468,7 @@ private:
 
 #endif // OPENVDB_AVAILABLE
 
-bool ConvertOpenVDBToSparseVolumeTexture(TArray<uint8>& SourceFile, const FOpenVDBImportOptions& ImportOptions, const FIntVector3& VolumeBoundsMin, FSparseVolumeRawSource& OutResult)
+bool ConvertOpenVDBToSparseVolumeTexture(TArray64<uint8>& SourceFile, const FOpenVDBImportOptions& ImportOptions, const FIntVector3& VolumeBoundsMin, FSparseVolumeRawSource& OutResult)
 {
 #if OPENVDB_AVAILABLE
 	FSparseVolumeRawSourceConstructionOpenVDBAdapter Adapter;
diff --git a/Engine/Source/Editor/SparseVolumeTexture/Private/SparseVolumeTextureOpenVDBUtility.h b/Engine/Source/Editor/SparseVolumeTexture/Private/SparseVolumeTextureOpenVDBUtility.h
--- a/Engine/Source/Editor/SparseVolumeTexture/Private/SparseVolumeTextureOpenVDBUtility.h
+++ b/Engine/Source/Editor/SparseVolumeTexture/Private/SparseVolumeTextureOpenVDBUtility.h
@@ -45,9 +45,9 @@ struct FOpenVDBToSVTConversionResult
 
 bool IsOpenVDBGridValid(const FOpenVDBGridInfo& GridInfo, const FString& Filename);
 
-bool GetOpenVDBGridInfo(TArray<uint8>& SourceFile, bool bCreateStrings, TArray<FOpenVDBGridInfo>* OutGridInfo);
+bool GetOpenVDBGridInfo(TArray64<uint8>& SourceFile, bool bCreateStrings, TArray<FOpenVDBGridInfo>* OutGridInfo);
 
-bool ConvertOpenVDBToSparseVolumeTexture(TArray<uint8>& SourceFile, const struct FOpenVDBImportOptions& ImportOptions, const FIntVector3& VolumeBoundsMin, struct FSparseVolumeRawSource& OutResult);
+bool ConvertOpenVDBToSparseVolumeTexture(TArray64<uint8>& SourceFile, const struct FOpenVDBImportOptions& ImportOptions, const FIntVector3& VolumeBoundsMin, struct FSparseVolumeRawSource& OutResult);
 
 const TCHAR* OpenVDBGridTypeToString(EOpenVDBGridType Type);
 
