Fixes for two small bugs when trying to import certain .vdb files. One is an off-by-one bug, the other one is an edge case for when a SVT frame is empty except for the null tile.

#rnx

diff --git a/Engine/Source/Editor/SparseVolumeTexture/Private/OpenVDBGridAdapter.h b/Engine/Source/Editor/SparseVolumeTexture/Private/OpenVDBGridAdapter.h
--- a/Engine/Source/Editor/SparseVolumeTexture/Private/OpenVDBGridAdapter.h
+++ b/Engine/Source/Editor/SparseVolumeTexture/Private/OpenVDBGridAdapter.h
@@ -144,12 +144,12 @@ public:
 				const openvdb::Coord AABBDimVDB = TileAABB.dim();
 				const FIntVector3 AABBDim = FIntVector3(AABBDimVDB.x(), AABBDimVDB.y(), AABBDimVDB.z());
 				
-				// Iterate over tile AABB and invoke callback for every (virtual) voxel. AABB.max is inclusive, so we need <= instead of <.
-				for (int32 Z = 0; Z <= AABBDim.Z; ++Z)
+				// Iterate over tile AABB and invoke callback for every (virtual) voxel.
+				for (int32 Z = 0; Z < AABBDim.Z; ++Z)
 				{
-					for (int32 Y = 0; Y <= AABBDim.Y; ++Y)
+					for (int32 Y = 0; Y < AABBDim.Y; ++Y)
 					{
-						for (int32 X = 0; X <= AABBDim.X; ++X)
+						for (int32 X = 0; X < AABBDim.X; ++X)
 						{
 							const FIntVector3 TileLocalCoord(Coord.X + X, Coord.Y + Y, Coord.Z + Z);
 							OnVisit(TileLocalCoord, NumComponents, VoxelValueComponents);
diff --git a/Engine/Source/Runtime/Engine/Private/SparseVolumeTexture/SparseVolumeTextureSceneProxy.cpp b/Engine/Source/Runtime/Engine/Private/SparseVolumeTexture/SparseVolumeTextureSceneProxy.cpp
--- a/Engine/Source/Runtime/Engine/Private/SparseVolumeTexture/SparseVolumeTextureSceneProxy.cpp
+++ b/Engine/Source/Runtime/Engine/Private/SparseVolumeTexture/SparseVolumeTextureSceneProxy.cpp
@@ -226,7 +226,7 @@ bool FSparseVolumeTextureRuntime::SetTileMappings(const TArrayView<const FSparse
 		PageTableRes = FIntVector3(FMath::Max(1, PageTableRes.X / 2), FMath::Max(1, PageTableRes.Y / 2), FMath::Max(1, PageTableRes.Z / 2));
 
 		const FSparseVolumeTextureTileMapping& Mapping = Mappings[MipLevel];
-		if (Mapping.NumPhysicalTiles <= 0)
+		if (!Mapping.TileIndices) // Note that we allow Mapping.NumPhysicalTiles == 0 in order to support a volume referencing just the null tile
 		{
 			continue;
 		}
