Fix for imported SVT assets having a name different from the package name, which causes them to not show up in the content browser.

#rnx

diff --git a/Engine/Source/Editor/SparseVolumeTexture/Private/SparseVolumeTextureFactory.cpp b/Engine/Source/Editor/SparseVolumeTexture/Private/SparseVolumeTextureFactory.cpp
--- a/Engine/Source/Editor/SparseVolumeTexture/Private/SparseVolumeTextureFactory.cpp
+++ b/Engine/Source/Editor/SparseVolumeTexture/Private/SparseVolumeTextureFactory.cpp
@@ -498,8 +498,7 @@ UObject* USparseVolumeTextureFactory::FactoryCreateFile(UClass* InClass, UObject
 		FScopedSlowTask ImportTask(1.0f, LOCTEXT("ImportingVDBStatic", "Importing static OpenVDB"));
 		ImportTask.MakeDialog(true);
 
-		FName NewName(InName.ToString() + TEXT("VDB"));
-		UStaticSparseVolumeTexture* StaticSVTexture = NewObject<UStaticSparseVolumeTexture>(InParent, UStaticSparseVolumeTexture::StaticClass(), NewName, Flags);
+		UStaticSparseVolumeTexture* StaticSVTexture = NewObject<UStaticSparseVolumeTexture>(InParent, UStaticSparseVolumeTexture::StaticClass(), InName, Flags);
 
 		ExpandVolumeBounds(ImportOptions, PreviewData.GridInfo, VolumeBoundsMin, VolumeBoundsMax);
 		StaticSVTexture->VolumeResolution = VolumeBoundsMax - VolumeBoundsMin;
@@ -535,8 +534,7 @@ UObject* USparseVolumeTextureFactory::FactoryCreateFile(UClass* InClass, UObject
 		// Data from original file is no longer needed; we iterate over all frames later
 		PreviewData.LoadedFile.Empty();
 
-		FName NewName(InName.ToString() + TEXT("VDBAnim"));
-		UAnimatedSparseVolumeTexture* AnimatedSVTexture = NewObject<UAnimatedSparseVolumeTexture>(InParent, UAnimatedSparseVolumeTexture::StaticClass(), NewName, Flags);
+		UAnimatedSparseVolumeTexture* AnimatedSVTexture = NewObject<UAnimatedSparseVolumeTexture>(InParent, UAnimatedSparseVolumeTexture::StaticClass(), InName, Flags);
 		
 		const int32 NumFrames = PreviewData.SequenceFilenames.Num();
 
