Moved SparseVolumeTexture initialization into a function, simplifying (and hiding implementation details of) SVT creation in factories etc.

#rnx

diff --git a/Engine/Source/Editor/SparseVolumeTexture/Private/SparseVolumeTextureFactory.cpp b/Engine/Source/Editor/SparseVolumeTexture/Private/SparseVolumeTextureFactory.cpp
--- a/Engine/Source/Editor/SparseVolumeTexture/Private/SparseVolumeTextureFactory.cpp
+++ b/Engine/Source/Editor/SparseVolumeTexture/Private/SparseVolumeTextureFactory.cpp
@@ -11,7 +11,6 @@
 #include "SparseVolumeTextureOpenVDBUtility.h"
 #include "OpenVDBImportOptions.h"
 
-#include "Serialization/EditorBulkDataWriter.h"
 #include "Misc/Paths.h"
 #include "Misc/ScopedSlowTask.h"
 #include "Misc/FileHelper.h"
@@ -498,12 +497,7 @@ UObject* USparseVolumeTextureFactory::FactoryCreateFile(UClass* InClass, UObject
 		FScopedSlowTask ImportTask(1.0f, LOCTEXT("ImportingVDBStatic", "Importing static OpenVDB"));
 		ImportTask.MakeDialog(true);
 
-		UStaticSparseVolumeTexture* StaticSVTexture = NewObject<UStaticSparseVolumeTexture>(InParent, UStaticSparseVolumeTexture::StaticClass(), InName, Flags);
-
 		ExpandVolumeBounds(ImportOptions, PreviewData.GridInfo, VolumeBoundsMin, VolumeBoundsMax);
-		StaticSVTexture->VolumeResolution = VolumeBoundsMax - VolumeBoundsMin;
-		StaticSVTexture->Frames.SetNum(1);
-		StaticSVTexture->NumMipLevels = ComputeNumMipLevels(VolumeBoundsMin, VolumeBoundsMax);
 
 		FSparseVolumeTextureData TextureData{};
 		const bool bConversionSuccess = ConvertOpenVDBToSparseVolumeTexture(PreviewData.LoadedFile, ImportOptions, VolumeBoundsMin, TextureData);
@@ -514,9 +508,13 @@ UObject* USparseVolumeTextureFactory::FactoryCreateFile(UClass* InClass, UObject
 			return nullptr;
 		}
 
-		// Serialize the raw source data into the asset object.
-		UE::Serialization::FEditorBulkDataWriter RawDataArchiveWriter(StaticSVTexture->Frames[0].RawData);
-		TextureData.Serialize(RawDataArchiveWriter);
+		UStaticSparseVolumeTexture* StaticSVTexture = NewObject<UStaticSparseVolumeTexture>(InParent, UStaticSparseVolumeTexture::StaticClass(), InName, Flags);
+		const bool bInitSuccess = StaticSVTexture->InitializeFromUncooked(TextureData);
+		if (!bInitSuccess)
+		{
+			UE_LOG(LogSparseVolumeTextureFactory, Error, TEXT("Failed to initialize SparseVolumeTexture: %s"), *Filename);
+			return nullptr;
+		}
 
 		if (ImportTask.ShouldCancel())
 		{
@@ -533,8 +531,6 @@ UObject* USparseVolumeTextureFactory::FactoryCreateFile(UClass* InClass, UObject
 		
 		// Data from original file is no longer needed; we iterate over all frames later
 		PreviewData.LoadedFile.Empty();
-
-		UAnimatedSparseVolumeTexture* AnimatedSVTexture = NewObject<UAnimatedSparseVolumeTexture>(InParent, UAnimatedSparseVolumeTexture::StaticClass(), InName, Flags);
 		
 		const int32 NumFrames = PreviewData.SequenceFilenames.Num();
 
@@ -542,7 +538,8 @@ UObject* USparseVolumeTextureFactory::FactoryCreateFile(UClass* InClass, UObject
 		ImportTask.MakeDialog(true);
 
 		// Allocate space for each frame
-		AnimatedSVTexture->Frames.SetNum(NumFrames);
+		TArray<FSparseVolumeTextureData> UncookedFramesData;
+		UncookedFramesData.SetNum(NumFrames);
 
 		std::atomic_bool bErrored = false;
 		std::atomic_bool bCanceled = false;
@@ -614,8 +611,6 @@ UObject* USparseVolumeTextureFactory::FactoryCreateFile(UClass* InClass, UObject
 
 		ImportTask.EnterProgressFrame(1.0f, LOCTEXT("ConvertingVDBAnim", "Converting OpenVDB animation"));
 
-		AnimatedSVTexture->NumMipLevels = ComputeNumMipLevels(VolumeBoundsMin, VolumeBoundsMax);
-
 		FEvent* AllTasksFinishedEvent = FPlatformProcess::GetSynchEventFromPool();
 		std::atomic_int FinishedTasksCounter = 0; // Will be incremented even if frame processing failed
 		std::atomic_int ProcessedFramesCounter = 0;
@@ -641,7 +636,7 @@ UObject* USparseVolumeTextureFactory::FactoryCreateFile(UClass* InClass, UObject
 			};
 
 			AsyncTask(ENamedThreads::AnyNormalThreadNormalTask, 
-				[FrameIdx, NumFrames, &PreviewData, &ImportOptions, &AnimatedSVTexture,
+				[FrameIdx, NumFrames, &PreviewData, &ImportOptions, &UncookedFramesData,
 				AllTasksFinishedEvent, &bErrored, &bCanceled, &FinishedTasksCounter, &ProcessedFramesCounter, &VolumeBoundsMin]()
 				{
 					// Ensure the FinishedTasksCounter will be incremented in all cases
@@ -675,9 +670,7 @@ UObject* USparseVolumeTextureFactory::FactoryCreateFile(UClass* InClass, UObject
 						return;
 					}
 
-					// Serialize the raw source data from this frame into the asset object.
-					UE::Serialization::FEditorBulkDataWriter RawDataArchiveWriter(AnimatedSVTexture->Frames[FrameIdx].RawData);
-					TextureData.Serialize(RawDataArchiveWriter);
+					UncookedFramesData[FrameIdx] = MoveTemp(TextureData);
 
 					// Increment ProcessedFramesCounter
 					ProcessedFramesCounter.fetch_add(1);
@@ -727,7 +720,13 @@ UObject* USparseVolumeTextureFactory::FactoryCreateFile(UClass* InClass, UObject
 			return nullptr;
 		}
 
-		AnimatedSVTexture->VolumeResolution = VolumeBoundsMax - VolumeBoundsMin;
+		UAnimatedSparseVolumeTexture* AnimatedSVTexture = NewObject<UAnimatedSparseVolumeTexture>(InParent, UAnimatedSparseVolumeTexture::StaticClass(), InName, Flags);
+		const bool bInitSuccess = AnimatedSVTexture->InitializeFromUncooked(UncookedFramesData);
+		if (!bInitSuccess)
+		{
+			UE_LOG(LogSparseVolumeTextureFactory, Error, TEXT("Failed to initialize SparseVolumeTexture: %s"), *Filename);
+			return nullptr;
+		}
 
 		ResultAssets.Add(AnimatedSVTexture);
 	}
@@ -751,5 +750,3 @@ UObject* USparseVolumeTextureFactory::FactoryCreateFile(UClass* InClass, UObject
 #endif // WITH_EDITORONLY_DATA
 
 #undef LOCTEXT_NAMESPACE
-
-#include "Serialization/EditorBulkDataWriter.h"
diff --git a/Engine/Source/Runtime/Engine/Classes/SparseVolumeTexture/SparseVolumeTexture.h b/Engine/Source/Runtime/Engine/Classes/SparseVolumeTexture/SparseVolumeTexture.h
--- a/Engine/Source/Runtime/Engine/Classes/SparseVolumeTexture/SparseVolumeTexture.h
+++ b/Engine/Source/Runtime/Engine/Classes/SparseVolumeTexture/SparseVolumeTexture.h
@@ -148,6 +148,8 @@ public:
 	UStreamableSparseVolumeTexture();
 	virtual ~UStreamableSparseVolumeTexture() = default;
 
+	virtual bool InitializeFromUncooked(const TArrayView<FSparseVolumeTextureData>& UncookedData, int32 NumMipLevels = INDEX_NONE /*Create entire mip chain by default*/);
+
 	//~ Begin UObject Interface.
 	virtual void PostLoad() override;
 	virtual void FinishDestroy() override;
@@ -201,15 +203,14 @@ public:
 	UStaticSparseVolumeTexture();
 	virtual ~UStaticSparseVolumeTexture() = default;
 
+	virtual bool InitializeFromUncooked(const TArrayView<FSparseVolumeTextureData>& UncookedData, int32 NumMipLevels = INDEX_NONE /*Create entire mip chain by default*/) override;
+	bool InitializeFromUncooked(FSparseVolumeTextureData& UncookedData, int32 NumMipLevels = INDEX_NONE /*Create entire mip chain by default*/);
+
 	//~ Begin USparseVolumeTexture Interface.
 	int32 GetNumFrames() const override { return 1; }
 	//~ End USparseVolumeTexture Interface.
 
 private:
-
-#if WITH_EDITOR
-	friend class USparseVolumeTextureFactory; // Importer
-#endif
 };
 
 // UAnimatedSparseVolumeTexture inherit from USparseVolumeTexture to be viewed using the first frame by default.
@@ -228,11 +229,6 @@ public:
 	//~ End USparseVolumeTexture Interface.
 
 private:
-
-#if WITH_EDITOR
-	friend class USparseVolumeTextureFactory; // Importer
-#endif
-	
 	int32 PreviewFrameIndex;
 	int32 PreviewMipLevel;
 };
diff --git a/Engine/Source/Runtime/Engine/Private/SparseVolumeTexture/SparseVolumeTexture.cpp b/Engine/Source/Runtime/Engine/Private/SparseVolumeTexture/SparseVolumeTexture.cpp
--- a/Engine/Source/Runtime/Engine/Private/SparseVolumeTexture/SparseVolumeTexture.cpp
+++ b/Engine/Source/Runtime/Engine/Private/SparseVolumeTexture/SparseVolumeTexture.cpp
@@ -29,6 +29,7 @@
 #include "Serialization/BulkDataReader.h"
 #include "Serialization/BulkDataWriter.h"
 #include "Serialization/EditorBulkDataReader.h"
+#include "Serialization/EditorBulkDataWriter.h"
 
 #include "ContentStreaming.h"
 
@@ -38,6 +39,20 @@ DEFINE_LOG_CATEGORY_STATIC(LogSparseVolumeTexture, Log, All);
 
 ////////////////////////////////////////////////////////////////////////////////////////////////
 
+static int32 SVTComputeNumMipLevels(const FIntVector3& InResolution)
+{
+	int32 Levels = 1;
+	FIntVector3 Resolution = InResolution;
+	while (Resolution.X > SPARSE_VOLUME_TILE_RES || Resolution.Y > SPARSE_VOLUME_TILE_RES || Resolution.Z > SPARSE_VOLUME_TILE_RES)
+	{
+		Resolution /= 2;
+		++Levels;
+	}
+	return Levels;
+};
+
+////////////////////////////////////////////////////////////////////////////////////////////////
+
 void FSparseVolumeTextureHeader::Serialize(FArchive& Ar)
 {
 	Ar << Version;
@@ -300,6 +315,66 @@ UStreamableSparseVolumeTexture::UStreamableSparseVolumeTexture(const FObjectInit
 {
 }
 
+bool UStreamableSparseVolumeTexture::InitializeFromUncooked(const TArrayView<FSparseVolumeTextureData>& InUncookedData, int32 InNumMipLevels)
+{
+	if (InUncookedData.IsEmpty())
+	{
+		UE_LOG(LogSparseVolumeTexture, Error, TEXT("Tried to initialize a SparseVolumeTexture with no frames"));
+		return false;
+	}
+
+	auto ExpandVolumeBounds = [](const FSparseVolumeTextureHeader& Header, FIntVector3& VolumeBoundsMin, FIntVector3& VolumeBoundsMax)
+	{
+		VolumeBoundsMin.X = FMath::Min(VolumeBoundsMin.X, Header.VirtualVolumeAABBMin.X);
+		VolumeBoundsMin.Y = FMath::Min(VolumeBoundsMin.Y, Header.VirtualVolumeAABBMin.Y);
+		VolumeBoundsMin.Z = FMath::Min(VolumeBoundsMin.Z, Header.VirtualVolumeAABBMin.Z);
+
+		VolumeBoundsMax.X = FMath::Max(VolumeBoundsMax.X, Header.VirtualVolumeAABBMax.X);
+		VolumeBoundsMax.Y = FMath::Max(VolumeBoundsMax.Y, Header.VirtualVolumeAABBMax.Y);
+		VolumeBoundsMax.Z = FMath::Max(VolumeBoundsMax.Z, Header.VirtualVolumeAABBMax.Z);
+	};
+
+	// Compute union of all frame AABBs
+	FIntVector3 VolumeBoundsMin = FIntVector3(INT32_MAX, INT32_MAX, INT32_MAX);
+	FIntVector3 VolumeBoundsMax = FIntVector3(INT32_MIN, INT32_MIN, INT32_MIN);
+	for (const FSparseVolumeTextureData& Data : InUncookedData)
+	{
+		ExpandVolumeBounds(Data.Header, VolumeBoundsMin, VolumeBoundsMax);
+	}
+
+	// The the minimum of the union of all frame AABBs should ideally be at (0, 0, 0), but it should also be fine if it is greater than that.
+	// A mimimum of less than (0, 0, 0) is not permitted.
+	if (VolumeBoundsMin.X < 0 || VolumeBoundsMin.Y < 0 || VolumeBoundsMin.Z < 0)
+	{
+		UE_LOG(LogSparseVolumeTexture, Error, TEXT("Tried to initialize a SparseVolumeTexture with a VirtualVolumeAABBMin < 0 (%i, %i, %i)"),
+			VolumeBoundsMin.X, VolumeBoundsMin.Y, VolumeBoundsMin.Z);
+		return false;
+	}
+	if (VolumeBoundsMin.X > 0 || VolumeBoundsMin.Y > 0 || VolumeBoundsMin.Z > 0)
+	{
+		UE_LOG(LogSparseVolumeTexture, Warning, TEXT("Initialized a SparseVolumeTexture with a VirtualVolumeAABBMin > 0 (%i, %i, %i). This wastes memory"),
+			VolumeBoundsMin.X, VolumeBoundsMin.Y, VolumeBoundsMin.Z);
+	}
+
+	VolumeResolution = VolumeBoundsMax;
+	Frames.SetNum(InUncookedData.Num());
+
+	const int32 NumMipLevelsFullMipChain = SVTComputeNumMipLevels(VolumeResolution);
+	check(NumMipLevelsFullMipChain > 0);
+
+	NumMipLevels = (InNumMipLevels <= INDEX_NONE) ? NumMipLevelsFullMipChain : FMath::Clamp(InNumMipLevels, 1, NumMipLevelsFullMipChain);
+
+#if WITH_EDITORONLY_DATA
+	for (int32 FrameIndex = 0; FrameIndex < Frames.Num(); ++FrameIndex)
+	{
+		UE::Serialization::FEditorBulkDataWriter RawDataArchiveWriter(Frames[FrameIndex].RawData);
+		InUncookedData[FrameIndex].Serialize(RawDataArchiveWriter);
+	}
+#endif
+
+	return true;
+}
+
 void UStreamableSparseVolumeTexture::PostLoad()
 {
 	Super::PostLoad();
@@ -577,6 +652,21 @@ UStaticSparseVolumeTexture::UStaticSparseVolumeTexture(const FObjectInitializer&
 {
 }
 
+bool UStaticSparseVolumeTexture::InitializeFromUncooked(const TArrayView<FSparseVolumeTextureData>& InUncookedData, int32 InNumMipLevels)
+{
+	if (InUncookedData.Num() > 1)
+	{
+		UE_LOG(LogSparseVolumeTexture, Error, TEXT("Tried to initialize a UStaticSparseVolumeTexture with more than 1 frame: %i"), InUncookedData.Num());
+		return false;
+	}
+	return Super::InitializeFromUncooked(InUncookedData, InNumMipLevels);
+}
+
+bool UStaticSparseVolumeTexture::InitializeFromUncooked(FSparseVolumeTextureData& InUncookedData, int32 InNumMipLevels)
+{
+	return Super::InitializeFromUncooked(MakeArrayView<FSparseVolumeTextureData>(&InUncookedData, 1), InNumMipLevels);
+}
+
 ////////////////////////////////////////////////////////////////////////////////////////////////
 
 UAnimatedSparseVolumeTexture::UAnimatedSparseVolumeTexture(const FObjectInitializer& ObjectInitializer)
