Added try/catch around openvdb file parsing to catch exceptions thrown on bad files. Moved grid typedefs for openvdb types into shared header.

#rnx

diff --git a/Engine/Source/Editor/SparseVolumeTexture/Private/OpenVDBGridAdapter.h b/Engine/Source/Editor/SparseVolumeTexture/Private/OpenVDBGridAdapter.h
--- a/Engine/Source/Editor/SparseVolumeTexture/Private/OpenVDBGridAdapter.h
+++ b/Engine/Source/Editor/SparseVolumeTexture/Private/OpenVDBGridAdapter.h
@@ -10,15 +10,6 @@
 
 #include "CoreMinimal.h"
 
-using FOpenVDBFloat1Grid = openvdb::FloatGrid;
-using FOpenVDBFloat2Grid = openvdb::Grid<openvdb::tree::Tree4<openvdb::Vec2f, 5, 4, 3>::Type>;
-using FOpenVDBFloat3Grid = openvdb::Vec3SGrid;
-using FOpenVDBFloat4Grid = openvdb::Grid<openvdb::tree::Tree4<openvdb::Vec4f, 5, 4, 3>::Type>;
-using FOpenVDBDouble1Grid = openvdb::DoubleGrid;
-using FOpenVDBDouble2Grid = openvdb::Grid<openvdb::tree::Tree4<openvdb::Vec2d, 5, 4, 3>::Type>;
-using FOpenVDBDouble3Grid = openvdb::Vec3DGrid;
-using FOpenVDBDouble4Grid = openvdb::Grid<openvdb::tree::Tree4<openvdb::Vec4d, 5, 4, 3>::Type>;
-
 template<typename ValueType>
 constexpr uint32 GetOpenVDBValueNumComponents()
 {
diff --git a/Engine/Source/Editor/SparseVolumeTexture/Private/SparseVolumeTextureOpenVDB.h b/Engine/Source/Editor/SparseVolumeTexture/Private/SparseVolumeTextureOpenVDB.h
--- a/Engine/Source/Editor/SparseVolumeTexture/Private/SparseVolumeTextureOpenVDB.h
+++ b/Engine/Source/Editor/SparseVolumeTexture/Private/SparseVolumeTextureOpenVDB.h
@@ -2,6 +2,9 @@
 
 #pragma once
 
+#include "HAL/Platform.h"
+#include "Misc/CoreMiscDefines.h" // UE_PUSH_MACRO seems to break without this
+
 #if PLATFORM_WINDOWS
 
 #define OPENVDB_AVAILABLE 1
@@ -36,6 +39,16 @@ UE_PUSH_MACRO("check")
 #include <openvdb/tools/Dense.h>
 #include <openvdb/math/Vec4.h>
 #include <openvdb/math/Half.h>
+#include <openvdb/Exceptions.h>
+
+using FOpenVDBFloat1Grid = openvdb::FloatGrid;
+using FOpenVDBFloat2Grid = openvdb::Grid<openvdb::tree::Tree4<openvdb::Vec2f, 5, 4, 3>::Type>;
+using FOpenVDBFloat3Grid = openvdb::Vec3SGrid;
+using FOpenVDBFloat4Grid = openvdb::Grid<openvdb::tree::Tree4<openvdb::Vec4f, 5, 4, 3>::Type>;
+using FOpenVDBDouble1Grid = openvdb::DoubleGrid;
+using FOpenVDBDouble2Grid = openvdb::Grid<openvdb::tree::Tree4<openvdb::Vec2d, 5, 4, 3>::Type>;
+using FOpenVDBDouble3Grid = openvdb::Vec3DGrid;
+using FOpenVDBDouble4Grid = openvdb::Grid<openvdb::tree::Tree4<openvdb::Vec4d, 5, 4, 3>::Type>;
 
 #if LOCAL_M_PI
 #undef M_PI
diff --git a/Engine/Source/Editor/SparseVolumeTexture/Private/SparseVolumeTextureOpenVDBUtility.cpp b/Engine/Source/Editor/SparseVolumeTexture/Private/SparseVolumeTextureOpenVDBUtility.cpp
--- a/Engine/Source/Editor/SparseVolumeTexture/Private/SparseVolumeTextureOpenVDBUtility.cpp
+++ b/Engine/Source/Editor/SparseVolumeTexture/Private/SparseVolumeTextureOpenVDBUtility.cpp
@@ -189,7 +189,16 @@ bool GetOpenVDBGridInfo(TArray64<uint8>& SourceFile, bool bCreateStrings, TArray
 #if OPENVDB_AVAILABLE
 	FArrayUint8StreamBuf StreamBuf(SourceFile);
 	std::istream IStream(&StreamBuf);
-	openvdb::io::Stream Stream(IStream, false /*delayLoad*/);
+	openvdb::io::Stream Stream;
+	try
+	{
+		Stream = openvdb::io::Stream(IStream, false /*delayLoad*/);
+	}
+	catch (const openvdb::Exception& Exception)
+	{
+		UE_LOG(LogSparseVolumeTextureOpenVDBUtility, Error, TEXT("Failed to read file due to exception: %s"), *FString(Exception.what()));
+		return false;
+	}
 
 	openvdb::GridPtrVecPtr Grids = Stream.getGrids();
 
@@ -305,7 +314,16 @@ public:
 		// Load file
 		FArrayUint8StreamBuf StreamBuf(SourceFile);
 		std::istream IStream(&StreamBuf);
-		openvdb::io::Stream Stream(IStream, false /*delayLoad*/);
+		openvdb::io::Stream Stream;
+		try
+		{
+			Stream = openvdb::io::Stream(IStream, false /*delayLoad*/);
+		}
+		catch (const openvdb::Exception& Exception)
+		{
+			UE_LOG(LogSparseVolumeTextureOpenVDBUtility, Error, TEXT("Failed to read file due to exception: %s"), *FString(Exception.what()));
+			return false;
+		}
 
 		// Check that the source grid indices are valid
 		openvdb::GridPtrVecPtr Grids = Stream.getGrids();
