Implemented support for importing .vdb grids with half float formats.

#rnx

diff --git a/Engine/Source/Editor/SparseVolumeTexture/Private/OpenVDBGridAdapter.h b/Engine/Source/Editor/SparseVolumeTexture/Private/OpenVDBGridAdapter.h
--- a/Engine/Source/Editor/SparseVolumeTexture/Private/OpenVDBGridAdapter.h
+++ b/Engine/Source/Editor/SparseVolumeTexture/Private/OpenVDBGridAdapter.h
@@ -16,6 +16,12 @@ constexpr uint32 GetOpenVDBValueNumComponents()
 	return (uint32)ValueType::size;
 }
 
+template<>
+constexpr uint32 GetOpenVDBValueNumComponents<FOpenVDBHalf>()
+{
+	return 1;
+}
+
 template<>
 constexpr uint32 GetOpenVDBValueNumComponents<float>()
 {
@@ -34,6 +40,12 @@ float GetOpenVDBValueComponent(const ValueType& Value, uint32 ComponentIndex)
 	return (float)Value[FMath::Min(ComponentIndex, (uint32)ValueType::size)];
 }
 
+template<>
+float GetOpenVDBValueComponent<FOpenVDBHalf>(const FOpenVDBHalf& Value, uint32 ComponentIndex)
+{
+	return (float)Value;
+}
+
 template<>
 float GetOpenVDBValueComponent<float>(const float& Value, uint32 ComponentIndex)
 {
@@ -52,6 +64,12 @@ ValueType GetOpenVDBMaxComponent(const ValueType& ValueA, const ValueType& Value
 	return openvdb::math::maxComponent(ValueA, ValueB);
 }
 
+template<>
+FOpenVDBHalf GetOpenVDBMaxComponent<FOpenVDBHalf>(const FOpenVDBHalf& ValueA, const FOpenVDBHalf& ValueB)
+{
+	return FOpenVDBHalf(FMath::Max((float)ValueA, (float)ValueB));
+}
+
 template<>
 float GetOpenVDBMaxComponent<float>(const float& ValueA, const float& ValueB)
 {
@@ -70,6 +88,12 @@ ValueType GetOpenVDBMinComponent(const ValueType& ValueA, const ValueType& Value
 	return openvdb::math::minComponent(ValueA, ValueB);
 }
 
+template<>
+FOpenVDBHalf GetOpenVDBMinComponent<FOpenVDBHalf>(const FOpenVDBHalf& ValueA, const FOpenVDBHalf& ValueB)
+{
+	return FOpenVDBHalf(FMath::Min((float)ValueA, (float)ValueB));
+}
+
 template<>
 float GetOpenVDBMinComponent<float>(const float& ValueA, const float& ValueB)
 {
@@ -198,7 +222,23 @@ TSharedPtr<IOpenVDBGridAdapterBase> CreateOpenVDBGridAdapterInternal(openvdb::Gr
 // Creates an adapter suitable for the type of the given grid or nullptr if the type is not supported.
 TSharedPtr<IOpenVDBGridAdapterBase> CreateOpenVDBGridAdapter(openvdb::GridBase::Ptr Grid)
 {
-	if (Grid->isType<FOpenVDBFloat1Grid>())
+	if (Grid->isType<FOpenVDBHalf1Grid>())
+	{
+		return CreateOpenVDBGridAdapterInternal<FOpenVDBHalf1Grid>(Grid);
+	}
+	else if (Grid->isType<FOpenVDBHalf2Grid>())
+	{
+		return CreateOpenVDBGridAdapterInternal<FOpenVDBHalf2Grid>(Grid);
+	}
+	else if (Grid->isType<FOpenVDBHalf3Grid>())
+	{
+		return CreateOpenVDBGridAdapterInternal<FOpenVDBHalf3Grid>(Grid);
+	}
+	else if (Grid->isType<FOpenVDBHalf4Grid>())
+	{
+		return CreateOpenVDBGridAdapterInternal<FOpenVDBHalf4Grid>(Grid);
+	}
+	else if (Grid->isType<FOpenVDBFloat1Grid>())
 	{
 		return CreateOpenVDBGridAdapterInternal<FOpenVDBFloat1Grid>(Grid);
 	}
diff --git a/Engine/Source/Editor/SparseVolumeTexture/Private/SparseVolumeTextureModule.cpp b/Engine/Source/Editor/SparseVolumeTexture/Private/SparseVolumeTextureModule.cpp
--- a/Engine/Source/Editor/SparseVolumeTexture/Private/SparseVolumeTextureModule.cpp
+++ b/Engine/Source/Editor/SparseVolumeTexture/Private/SparseVolumeTextureModule.cpp
@@ -8,11 +8,32 @@
 
 IMPLEMENT_MODULE(FSparseVolumeTextureModule, SparseVolumeTexture);
 
+template<typename T>
+static void RegisterOpenVDBGrid()
+{
+	if (!T::isRegistered())
+	{
+		T::registerGrid();
+	}
+}
+
 void FSparseVolumeTextureModule::StartupModule()
 {
 #if PLATFORM_WINDOWS
 	// Global registration of  the vdb types.
 	openvdb::initialize();
+	RegisterOpenVDBGrid<FOpenVDBHalf1Grid>();
+	RegisterOpenVDBGrid<FOpenVDBHalf2Grid>();
+	RegisterOpenVDBGrid<FOpenVDBHalf3Grid>();
+	RegisterOpenVDBGrid<FOpenVDBHalf4Grid>();
+	RegisterOpenVDBGrid<FOpenVDBFloat1Grid>();
+	RegisterOpenVDBGrid<FOpenVDBFloat2Grid>();
+	RegisterOpenVDBGrid<FOpenVDBFloat3Grid>();
+	RegisterOpenVDBGrid<FOpenVDBFloat4Grid>();
+	RegisterOpenVDBGrid<FOpenVDBDouble1Grid>();
+	RegisterOpenVDBGrid<FOpenVDBDouble2Grid>();
+	RegisterOpenVDBGrid<FOpenVDBDouble3Grid>();
+	RegisterOpenVDBGrid<FOpenVDBDouble4Grid>();
 #endif
 }
 
diff --git a/Engine/Source/Editor/SparseVolumeTexture/Private/SparseVolumeTextureOpenVDB.h b/Engine/Source/Editor/SparseVolumeTexture/Private/SparseVolumeTextureOpenVDB.h
--- a/Engine/Source/Editor/SparseVolumeTexture/Private/SparseVolumeTextureOpenVDB.h
+++ b/Engine/Source/Editor/SparseVolumeTexture/Private/SparseVolumeTextureOpenVDB.h
@@ -41,6 +41,11 @@ UE_PUSH_MACRO("check")
 #include <openvdb/math/Half.h>
 #include <openvdb/Exceptions.h>
 
+using FOpenVDBHalf = openvdb::math::internal::half;
+using FOpenVDBHalf1Grid = openvdb::Grid<openvdb::tree::Tree4<FOpenVDBHalf, 5, 4, 3>::Type>;
+using FOpenVDBHalf2Grid = openvdb::Grid<openvdb::tree::Tree4<openvdb::math::Vec2<FOpenVDBHalf>, 5, 4, 3>::Type>;
+using FOpenVDBHalf3Grid = openvdb::Grid<openvdb::tree::Tree4<openvdb::math::Vec3<FOpenVDBHalf>, 5, 4, 3>::Type>;
+using FOpenVDBHalf4Grid = openvdb::Grid<openvdb::tree::Tree4<openvdb::math::Vec4<FOpenVDBHalf>, 5, 4, 3>::Type>;
 using FOpenVDBFloat1Grid = openvdb::FloatGrid;
 using FOpenVDBFloat2Grid = openvdb::Grid<openvdb::tree::Tree4<openvdb::Vec2f, 5, 4, 3>::Type>;
 using FOpenVDBFloat3Grid = openvdb::Vec3SGrid;
diff --git a/Engine/Source/Editor/SparseVolumeTexture/Private/SparseVolumeTextureOpenVDBUtility.cpp b/Engine/Source/Editor/SparseVolumeTexture/Private/SparseVolumeTextureOpenVDBUtility.cpp
--- a/Engine/Source/Editor/SparseVolumeTexture/Private/SparseVolumeTextureOpenVDBUtility.cpp
+++ b/Engine/Source/Editor/SparseVolumeTexture/Private/SparseVolumeTextureOpenVDBUtility.cpp
@@ -116,7 +116,27 @@ static FOpenVDBGridInfo GetOpenVDBGridInfo(openvdb::GridBase::Ptr Grid, uint32 G
 	}
 
 	// Figure out the type/format of the grid
-	if (Grid->isType<FOpenVDBFloat1Grid>())
+	if (Grid->isType<FOpenVDBHalf1Grid>())
+	{
+		GridInfo.NumComponents = 1;
+		GridInfo.Type = EOpenVDBGridType::Half;
+	}
+	else if (Grid->isType<FOpenVDBHalf2Grid>())
+	{
+		GridInfo.NumComponents = 2;
+		GridInfo.Type = EOpenVDBGridType::Half2;
+	}
+	else if (Grid->isType<FOpenVDBHalf3Grid>())
+	{
+		GridInfo.NumComponents = 3;
+		GridInfo.Type = EOpenVDBGridType::Half3;
+	}
+	else if (Grid->isType<FOpenVDBHalf4Grid>())
+	{
+		GridInfo.NumComponents = 4;
+		GridInfo.Type = EOpenVDBGridType::Half4;
+	}
+	else if (Grid->isType<FOpenVDBFloat1Grid>())
 	{
 		GridInfo.NumComponents = 1;
 		GridInfo.Type = EOpenVDBGridType::Float;
@@ -502,6 +522,14 @@ const TCHAR* OpenVDBGridTypeToString(EOpenVDBGridType Type)
 {
 	switch (Type)
 	{
+	case EOpenVDBGridType::Half:
+		return TEXT("Half");
+	case EOpenVDBGridType::Half2:
+		return TEXT("Half2");
+	case EOpenVDBGridType::Half3:
+		return TEXT("Half3");
+	case EOpenVDBGridType::Half4:
+		return TEXT("Half4");
 	case EOpenVDBGridType::Float:
 		return TEXT("Float");
 	case EOpenVDBGridType::Float2:
diff --git a/Engine/Source/Editor/SparseVolumeTexture/Private/SparseVolumeTextureOpenVDBUtility.h b/Engine/Source/Editor/SparseVolumeTexture/Private/SparseVolumeTextureOpenVDBUtility.h
--- a/Engine/Source/Editor/SparseVolumeTexture/Private/SparseVolumeTextureOpenVDBUtility.h
+++ b/Engine/Source/Editor/SparseVolumeTexture/Private/SparseVolumeTextureOpenVDBUtility.h
@@ -9,14 +9,18 @@
 enum class EOpenVDBGridType : uint8
 {
 	Unknown = 0,
-	Float = 1,
-	Float2 = 2,
-	Float3 = 3,
-	Float4 = 4,
-	Double = 5,
-	Double2 = 6,
-	Double3 = 7,
-	Double4 = 8,
+	Half,
+	Half2,
+	Half3,
+	Half4,
+	Float,
+	Float2,
+	Float3,
+	Float4,
+	Double,
+	Double2,
+	Double3,
+	Double4,
 };
 
 struct FOpenVDBGridInfo
