Implemented support for reimporting SparseVolumeTextures. Also fixed some inconsistent dllexport macros which initially caused problems when trying to call generated functions of the SVT classes in the factory (which resides in a different module).

#rnx

diff --git a/Engine/Source/Editor/SparseVolumeTexture/Private/SparseVolumeTextureFactory.cpp b/Engine/Source/Editor/SparseVolumeTexture/Private/SparseVolumeTextureFactory.cpp
--- a/Engine/Source/Editor/SparseVolumeTexture/Private/SparseVolumeTextureFactory.cpp
+++ b/Engine/Source/Editor/SparseVolumeTexture/Private/SparseVolumeTextureFactory.cpp
@@ -18,6 +18,7 @@
 #include "Async/ParallelFor.h"
 
 #include "Editor.h"
+#include "EditorFramework/AssetImportData.h"
 
 #include "OpenVDBImportWindow.h"
 #include "HAL/PlatformApplicationMisc.h"
@@ -342,7 +343,7 @@ USparseVolumeTextureFactory::USparseVolumeTextureFactory(const FObjectInitialize
 	bCreateNew = true;
 	bEditAfterNew = true;
 	bEditorImport = true;
-	SupportedClass = USparseVolumeTexture::StaticClass();
+	SupportedClass = nullptr; // This factory supports multiple classes, so SupportedClass needs to be nullptr
 
 	Formats.Add(TEXT("vdb;OpenVDB Format"));
 }
@@ -388,11 +389,12 @@ UObject* USparseVolumeTextureFactory::FactoryCreateNew(UClass* InClass, UObject*
 
 bool USparseVolumeTextureFactory::DoesSupportClass(UClass* Class)
 {
-	return Class == USparseVolumeTexture::StaticClass();
+	return Class == USparseVolumeTexture::StaticClass() || Class == UStaticSparseVolumeTexture::StaticClass() || Class == UAnimatedSparseVolumeTexture::StaticClass();
 }
 
 UClass* USparseVolumeTextureFactory::ResolveSupportedClass()
 {
+	// SVT_TODO: Do we need to return UStaticSparseVolumeTexture::StaticClass() or UAnimatedSparseVolumeTexture::StaticClass() here instead? Using the base class seems to work.
 	return USparseVolumeTexture::StaticClass();
 }
 
@@ -413,6 +415,70 @@ void USparseVolumeTextureFactory::CleanUp()
 
 UObject* USparseVolumeTextureFactory::FactoryCreateFile(UClass* InClass, UObject* InParent, FName InName, EObjectFlags Flags, const FString& Filename,
 	const TCHAR* Parms, FFeedbackContext* Warn, bool& bOutOperationCanceled)
+{
+	return ImportInternal(InClass, InParent, InName, Flags, Filename, Parms, bOutOperationCanceled);
+}
+
+bool USparseVolumeTextureFactory::CanReimport(UObject* Obj, TArray<FString>& OutFilenames)
+{
+#if OPENVDB_AVAILABLE
+	UStreamableSparseVolumeTexture* StreamableSVT = Cast<UStreamableSparseVolumeTexture>(Obj);
+	if (StreamableSVT && StreamableSVT->AssetImportData)
+	{
+		StreamableSVT->AssetImportData->ExtractFilenames(OutFilenames);
+		return true;
+	}
+#endif // OPENVDB_AVAILABLE
+	return false;
+}
+
+void USparseVolumeTextureFactory::SetReimportPaths(UObject* Obj, const TArray<FString>& NewReimportPaths)
+{
+#if OPENVDB_AVAILABLE
+	UStreamableSparseVolumeTexture* StreamableSVT = Cast<UStreamableSparseVolumeTexture>(Obj);
+	if (StreamableSVT && ensure(NewReimportPaths.Num() == 1))
+	{
+		StreamableSVT->AssetImportData->UpdateFilenameOnly(NewReimportPaths[0]);
+	}
+#endif // OPENVDB_AVAILABLE
+}
+
+EReimportResult::Type USparseVolumeTextureFactory::Reimport(UObject* Obj)
+{
+#if OPENVDB_AVAILABLE
+	UStreamableSparseVolumeTexture* StreamableSVT = Cast<UStreamableSparseVolumeTexture>(Obj);
+	if (!StreamableSVT)
+	{
+		return EReimportResult::Failed;
+	}
+
+	// Make sure file is valid and exists
+	const FString Filename = StreamableSVT->AssetImportData->GetFirstFilename();
+	if (!Filename.Len() || IFileManager::Get().FileSize(*Filename) == INDEX_NONE || !FactoryCanImport(Filename))
+	{
+		UE_LOG(LogSparseVolumeTextureFactory, Error, TEXT("Reimport failed! Filename '%s' is invalid or no such file exists."), *Filename);
+		return EReimportResult::Failed;
+	}
+
+	bool OutCanceled = false;
+	if (!ImportInternal(StreamableSVT->GetClass(), StreamableSVT->GetOuter(), *StreamableSVT->GetName(), RF_Public | RF_Standalone, Filename, nullptr, OutCanceled))
+	{
+		if (OutCanceled)
+		{
+			return EReimportResult::Cancelled;
+		}
+
+		return EReimportResult::Failed;
+	}
+
+	return EReimportResult::Succeeded;
+#else
+	UE_LOG(LogSparseVolumeTextureFactory, Error, TEXT("Cannot import OpenVDB asset any platform other than Windows."));
+	return EReimportResult::Failed;
+#endif // OPENVDB_AVAILABLE
+}
+
+UObject* USparseVolumeTextureFactory::ImportInternal(UClass* InClass, UObject* InParent, FName InName, EObjectFlags Flags, const FString& Filename, const TCHAR* Parms, bool& bOutOperationCanceled)
 {
 #if OPENVDB_AVAILABLE
 
@@ -434,7 +500,7 @@ UObject* USparseVolumeTextureFactory::FactoryCreateFile(UClass* InClass, UObject
 	}
 
 	FOpenVDBImportOptions ImportOptions = PreviewData.DefaultImportOptions;
-	
+
 	if (!bIsUnattended)
 	{
 		// Show dialog for import options
@@ -527,10 +593,10 @@ UObject* USparseVolumeTextureFactory::FactoryCreateFile(UClass* InClass, UObject
 	else
 	{
 		// Import as an animated sparse volume texture asset.
-		
+
 		// Data from original file is no longer needed; we iterate over all frames later
 		PreviewData.LoadedFile.Empty();
-		
+
 		const int32 NumFrames = PreviewData.SequenceFilenames.Num();
 
 		FScopedSlowTask ImportTask(NumFrames + 1, LOCTEXT("ImportingVDBAnim", "Importing OpenVDB animation"));
@@ -623,10 +689,10 @@ UObject* USparseVolumeTextureFactory::FactoryCreateFile(UClass* InClass, UObject
 				std::atomic_int& Counter;
 				int32 MaxValue;
 				FEvent* Event;
-				explicit FScopedIncrementer(std::atomic_int& InCounter, int32 InMaxValue, FEvent* InEvent) 
+				explicit FScopedIncrementer(std::atomic_int& InCounter, int32 InMaxValue, FEvent* InEvent)
 					: Counter(InCounter), MaxValue(InMaxValue), Event(InEvent) {}
-				~FScopedIncrementer() 
-				{ 
+				~FScopedIncrementer()
+				{
 					if ((Counter.fetch_add(1) + 1) == MaxValue)
 					{
 						Event->Trigger();
@@ -634,7 +700,7 @@ UObject* USparseVolumeTextureFactory::FactoryCreateFile(UClass* InClass, UObject
 				}
 			};
 
-			AsyncTask(ENamedThreads::AnyNormalThreadNormalTask, 
+			AsyncTask(ENamedThreads::AnyNormalThreadNormalTask,
 				[FrameIdx, NumFrames, &PreviewData, &ImportOptions, &UncookedFramesData,
 				AllTasksFinishedEvent, &bErrored, &bCanceled, &FinishedTasksCounter, &ProcessedFramesCounter, &VolumeBoundsMin]()
 				{
@@ -680,7 +746,7 @@ UObject* USparseVolumeTextureFactory::FactoryCreateFile(UClass* InClass, UObject
 		{
 			int NumFinishedTasks = 0;
 			int NumProcessedFrames = 0;
-			
+
 			while (NumFinishedTasks < NumFrames)
 			{
 				// We can't block here because we want to regularly update the progress bar and check for user input.
@@ -732,6 +798,7 @@ UObject* USparseVolumeTextureFactory::FactoryCreateFile(UClass* InClass, UObject
 
 	// Now notify the system about the imported/updated/created assets
 	check(ResultAssets.Num() == 1);
+	CastChecked<UStreamableSparseVolumeTexture>(ResultAssets[0])->AssetImportData->Update(Filename);
 	GEditor->GetEditorSubsystem<UImportSubsystem>()->BroadcastAssetPostImport(this, ResultAssets[0]);
 
 	return ResultAssets[0];
@@ -743,7 +810,6 @@ UObject* USparseVolumeTextureFactory::FactoryCreateFile(UClass* InClass, UObject
 	return nullptr;
 
 #endif // OPENVDB_AVAILABLE
-
 }
 
 #endif // WITH_EDITORONLY_DATA
diff --git a/Engine/Source/Editor/SparseVolumeTexture/Private/SparseVolumeTextureFactory.h b/Engine/Source/Editor/SparseVolumeTexture/Private/SparseVolumeTextureFactory.h
--- a/Engine/Source/Editor/SparseVolumeTexture/Private/SparseVolumeTextureFactory.h
+++ b/Engine/Source/Editor/SparseVolumeTexture/Private/SparseVolumeTextureFactory.h
@@ -6,12 +6,13 @@
 
 #include "CoreMinimal.h"
 #include "Factories/Factory.h"
+#include "EditorReimportHandler.h"
 
 #include "SparseVolumeTextureFactory.generated.h"
 
 // Responsible for creating and importing Sparse Volume Texture objects
 UCLASS(hidecategories = Object, MinimalAPI)
-class USparseVolumeTextureFactory : public UFactory
+class USparseVolumeTextureFactory : public UFactory, public FReimportHandler
 {
 	GENERATED_UCLASS_BODY()
 
@@ -43,9 +44,16 @@ public:
 
 	//~ End UFactory Interface
 
+	//~ Begin FReimportHandler Interface
+	virtual bool CanReimport(UObject* Obj, TArray<FString>& OutFilenames) override;
+	virtual void SetReimportPaths(UObject* Obj, const TArray<FString>& NewReimportPaths) override;
+	virtual EReimportResult::Type Reimport(UObject* Obj) override;
+	//~ End FReimportHandler Interface
+
 protected:
 
 private:
+	UObject* ImportInternal(UClass* InClass, UObject* InParent, FName InName, EObjectFlags Flags, const FString& Filename, const TCHAR* Parms, bool& bOutOperationCanceled);
 };
 
 #endif // WITH_EDITOR
diff --git a/Engine/Source/Runtime/Engine/Classes/SparseVolumeTexture/SparseVolumeTexture.h b/Engine/Source/Runtime/Engine/Classes/SparseVolumeTexture/SparseVolumeTexture.h
--- a/Engine/Source/Runtime/Engine/Classes/SparseVolumeTexture/SparseVolumeTexture.h
+++ b/Engine/Source/Runtime/Engine/Classes/SparseVolumeTexture/SparseVolumeTexture.h
@@ -186,7 +186,7 @@ enum ESparseVolumeTextureShaderUniform
 };
 
 // SparseVolumeTexture base interface to communicate with material graph and shader bindings.
-UCLASS(ClassGroup = Rendering, BlueprintType, MinimalAPI)
+UCLASS(MinimalAPI, ClassGroup = Rendering, BlueprintType)
 class USparseVolumeTexture : public UObject
 {
 	GENERATED_UCLASS_BODY()
@@ -238,7 +238,7 @@ public:
 };
 
 // Represents a frame in a SparseVolumeTexture sequence and owns the actual data needed for rendering. Is owned by a UStreamableSparseVolumeTexture object.
-UCLASS(ClassGroup = Rendering, BlueprintType, MinimalAPI)
+UCLASS(MinimalAPI, ClassGroup = Rendering, BlueprintType)
 class USparseVolumeTextureFrame : public USparseVolumeTexture
 {
 	GENERATED_UCLASS_BODY()
@@ -316,7 +316,7 @@ private:
 };
 
 // Represents a streamable SparseVolumeTexture asset and serves as base class for UStaticSparseVolumeTexture and UAnimatedSparseVolumeTexture. It has an array of USparseVolumeTextureFrame.
-UCLASS(ClassGroup = Rendering, BlueprintType)//, hidecategories = (Object))
+UCLASS(MinimalAPI, ClassGroup = Rendering, BlueprintType)
 class UStreamableSparseVolumeTexture : public USparseVolumeTexture
 {
 	GENERATED_UCLASS_BODY()
@@ -357,15 +357,20 @@ public:
 	UPROPERTY(EditAnywhere, BlueprintReadWrite, Category = "Texture", meta = (DisplayName = "Local DDC Only"), AssetRegistrySearchable, AdvancedDisplay)
 	bool bLocalDDCOnly = true;
 
+#if WITH_EDITORONLY_DATA
+	UPROPERTY(VisibleAnywhere, Instanced, Category = ImportSettings)
+	TObjectPtr<class UAssetImportData> AssetImportData;
+#endif // WITH_EDITORONLY_DATA
+
 	UStreamableSparseVolumeTexture();
 	virtual ~UStreamableSparseVolumeTexture() = default;
 
 	// Multi-phase initialization: Call BeginInitialize(), then call AppendFrame() for each frame to add and then finish initialization with a call to EndInitialize().
 	// The NumExpectedFrames parameter on BeginInitialize() just serves as a potential optimization to reserve memory for the frames to be appended
 	// and doesn't need to match the exact number if it is not known at the time.
-	virtual bool BeginInitialize(int32 NumExpectedFrames);
-	virtual bool AppendFrame(UE::SVT::FTextureData& UncookedFrame);
-	virtual bool EndInitialize(int32 NumMipLevels = INDEX_NONE /*Create entire mip chain by default*/);
+	ENGINE_API virtual bool BeginInitialize(int32 NumExpectedFrames);
+	ENGINE_API virtual bool AppendFrame(UE::SVT::FTextureData& UncookedFrame);
+	ENGINE_API virtual bool EndInitialize(int32 NumMipLevels = INDEX_NONE /*Create entire mip chain by default*/);
 
 	// Convenience function wrapping the multi-phase initialization functions above
 	virtual bool Initialize(const TArrayView<UE::SVT::FTextureData>& UncookedData, int32 NumMipLevels = INDEX_NONE /*Create entire mip chain by default*/);
@@ -373,14 +378,16 @@ public:
 	USparseVolumeTextureFrame* GetFrame(int32 FrameIndex) const { return Frames.IsValidIndex(FrameIndex) ? Frames[FrameIndex] : nullptr; }
 
 	//~ Begin UObject Interface.
-	virtual void PostLoad() override;
-	virtual void FinishDestroy() override;
-	virtual void BeginDestroy() override;
-	virtual void Serialize(FArchive& Ar) override;
+	ENGINE_API virtual void PostInitProperties() override;
+	ENGINE_API virtual void PostLoad() override;
+	ENGINE_API virtual void FinishDestroy() override;
+	ENGINE_API virtual void BeginDestroy() override;
+	ENGINE_API virtual void Serialize(FArchive& Ar) override;
 #if WITH_EDITOR
-	virtual void PostEditChangeProperty(FPropertyChangedEvent& PropertyChangedEvent) override;
+	ENGINE_API virtual void PostEditChangeProperty(FPropertyChangedEvent& PropertyChangedEvent) override;
 #endif
-	virtual void GetResourceSizeEx(FResourceSizeEx& CumulativeResourceSize) override;
+	ENGINE_API virtual void GetResourceSizeEx(FResourceSizeEx& CumulativeResourceSize) override;
+	ENGINE_API virtual void GetAssetRegistryTags(TArray<FAssetRegistryTag>& OutTags) const override;
 	//~ End UObject Interface.
 
 	//~ Begin USparseVolumeTexture Interface.
@@ -426,7 +433,7 @@ protected:
 };
 
 // Represents a streamable SparseVolumeTexture asset with a single frame. Although there is only a single frame, it is still recommended to use USparseVolumeTextureFrame::GetFrameAndIssueStreamingRequest().
-UCLASS(MinimalAPI, ClassGroup = Rendering, BlueprintType)//, hidecategories = (Object), MinimalAPI)
+UCLASS(MinimalAPI, ClassGroup = Rendering, BlueprintType)
 class UStaticSparseVolumeTexture : public UStreamableSparseVolumeTexture
 {
 	GENERATED_UCLASS_BODY()
@@ -449,7 +456,7 @@ private:
 };
 
 // Represents a streamable SparseVolumeTexture with one or more frames. Use USparseVolumeTextureFrame::GetFrameAndIssueStreamingRequest() to bind extract a particular frame to be used for rendering.
-UCLASS(MinimalAPI, ClassGroup = Rendering, BlueprintType)//, hidecategories = (Object), MinimalAPI)
+UCLASS(MinimalAPI, ClassGroup = Rendering, BlueprintType)
 class UAnimatedSparseVolumeTexture : public UStreamableSparseVolumeTexture
 {
 	GENERATED_UCLASS_BODY()
@@ -468,7 +475,7 @@ private:
 };
 
 // Utility (blueprint) class for controlling SparseVolumeTexture playback.
-UCLASS(ClassGroup = Rendering, BlueprintType)//, hidecategories = (Object), MinimalAPI)
+UCLASS(MinimalAPI, ClassGroup = Rendering, BlueprintType)
 class UAnimatedSparseVolumeTextureController : public UObject
 {
 	GENERATED_UCLASS_BODY()
diff --git a/Engine/Source/Runtime/Engine/Private/SparseVolumeTexture/SparseVolumeTexture.cpp b/Engine/Source/Runtime/Engine/Private/SparseVolumeTexture/SparseVolumeTexture.cpp
--- a/Engine/Source/Runtime/Engine/Private/SparseVolumeTexture/SparseVolumeTexture.cpp
+++ b/Engine/Source/Runtime/Engine/Private/SparseVolumeTexture/SparseVolumeTexture.cpp
@@ -25,6 +25,7 @@
 #include "DerivedDataCache.h"
 #include "DerivedDataRequestOwner.h"
 #endif
+#include "EditorFramework/AssetImportData.h"
 
 #include "Serialization/LargeMemoryReader.h"
 #include "Serialization/LargeMemoryWriter.h"
@@ -1074,6 +1075,18 @@ bool UStreamableSparseVolumeTexture::Initialize(const TArrayView<UE::SVT::FTextu
 	return true;
 }
 
+void UStreamableSparseVolumeTexture::PostInitProperties()
+{
+	Super::PostInitProperties();
+
+#if WITH_EDITORONLY_DATA
+	if (!HasAnyFlags(RF_ClassDefaultObject))
+	{
+		AssetImportData = NewObject<UAssetImportData>(this, TEXT("AssetImportData"));
+	}
+#endif
+}
+
 void UStreamableSparseVolumeTexture::PostLoad()
 {
 	Super::PostLoad();
@@ -1145,6 +1158,18 @@ void UStreamableSparseVolumeTexture::GetResourceSizeEx(FResourceSizeEx& Cumulati
 	CumulativeResourceSize.AddDedicatedVideoMemoryBytes(SizeGPU);
 }
 
+void UStreamableSparseVolumeTexture::GetAssetRegistryTags(TArray<FAssetRegistryTag>& OutTags) const
+{
+	Super::GetAssetRegistryTags(OutTags);
+
+#if WITH_EDITORONLY_DATA
+	if (AssetImportData)
+	{
+		OutTags.Add(FAssetRegistryTag(SourceFileTagName(), AssetImportData->GetSourceData().ToJson(), FAssetRegistryTag::TT_Hidden));
+	}
+#endif
+}
+
 #if WITH_EDITORONLY_DATA
 void UStreamableSparseVolumeTexture::RecacheFrames()
 {
