Implemented automatic asset renaming when importing .vdb sequences as an animated SparseVolumeTexture. The resulting asset name will no longer have the number postfix of the original file used to import the entire sequence.
SVT assets now also have an additional property to show the number of frames in the editor UI.

#rnx

diff --git a/Engine/Source/Editor/SparseVolumeTexture/Private/SparseVolumeTextureFactory.cpp b/Engine/Source/Editor/SparseVolumeTexture/Private/SparseVolumeTextureFactory.cpp
--- a/Engine/Source/Editor/SparseVolumeTexture/Private/SparseVolumeTextureFactory.cpp
+++ b/Engine/Source/Editor/SparseVolumeTexture/Private/SparseVolumeTextureFactory.cpp
@@ -70,6 +70,21 @@ static void ComputeDefaultOpenVDBGridAssignment(const TArray<TSharedPtr<FOpenVDB
 	ImportOptions->bIsSequence = NumFiles > 1;
 }
 
+static FString GetVDBSequenceBaseFileName(const FString& FileName, bool bDiscardUnderscore)
+{
+	const FString CleanFileName = FPaths::GetCleanFilename(FileName);
+	const FString CleanFileNameWithoutExt = CleanFileName.LeftChop(4);
+	const int32 LastNonDigitIndex = CleanFileNameWithoutExt.FindLastCharByPredicate([&](TCHAR Letter) 
+		{ 
+			const bool bIsDigit = FChar::IsDigit(Letter);
+			const bool bIsUnderscore = bDiscardUnderscore && Letter == TEXT('_');
+			return !bIsDigit && !bIsUnderscore;
+		}) + 1;
+	const int32 DigitCount = CleanFileNameWithoutExt.Len() - LastNonDigitIndex;
+	const FString CleanFileNameWithoutSuffix = CleanFileNameWithoutExt.LeftChop(DigitCount);
+	return CleanFileNameWithoutSuffix;
+}
+
 static TArray<FString> FindOpenVDBSequenceFileNames(const FString& Filename)
 {
 	TArray<FString> SequenceFilenames;
@@ -85,13 +100,7 @@ static TArray<FString> FindOpenVDBSequenceFileNames(const FString& Filename)
 	{
 		const FString Path = FPaths::GetPath(Filename);
 		const FString CleanFilename = FPaths::GetCleanFilename(Filename);
-		FString CleanFilenameWithoutSuffix;
-		{
-			const FString CleanFilenameWithoutExt = CleanFilename.LeftChop(4);
-			const int32 LastNonDigitIndex = CleanFilenameWithoutExt.FindLastCharByPredicate([](TCHAR Letter) { return !FChar::IsDigit(Letter); }) + 1;
-			const int32 DigitCount = CleanFilenameWithoutExt.Len() - LastNonDigitIndex;
-			CleanFilenameWithoutSuffix = CleanFilenameWithoutExt.LeftChop(CleanFilenameWithoutExt.Len() - LastNonDigitIndex);
-		}
+		const FString CleanFilenameWithoutSuffix = GetVDBSequenceBaseFileName(Filename, false /*bDiscardUnderscore*/);
 
 		// Find all files potentially part of the sequence
 		TArray<FString> PotentialSequenceFilenames;
@@ -416,7 +425,7 @@ void USparseVolumeTextureFactory::CleanUp()
 UObject* USparseVolumeTextureFactory::FactoryCreateFile(UClass* InClass, UObject* InParent, FName InName, EObjectFlags Flags, const FString& Filename,
 	const TCHAR* Parms, FFeedbackContext* Warn, bool& bOutOperationCanceled)
 {
-	return ImportInternal(InClass, InParent, InName, Flags, Filename, Parms, bOutOperationCanceled);
+	return ImportInternal(InClass, InParent, InName, Flags, Filename, Parms, bOutOperationCanceled, false /*bIsReimport*/);
 }
 
 bool USparseVolumeTextureFactory::CanReimport(UObject* Obj, TArray<FString>& OutFilenames)
@@ -461,7 +470,7 @@ EReimportResult::Type USparseVolumeTextureFactory::Reimport(UObject* Obj)
 	}
 
 	bool OutCanceled = false;
-	if (!ImportInternal(StreamableSVT->GetClass(), StreamableSVT->GetOuter(), *StreamableSVT->GetName(), RF_Public | RF_Standalone, Filename, nullptr, OutCanceled))
+	if (!ImportInternal(StreamableSVT->GetClass(), StreamableSVT->GetOuter(), *StreamableSVT->GetName(), RF_Public | RF_Standalone, Filename, nullptr, OutCanceled, true /*bIsReimport*/))
 	{
 		if (OutCanceled)
 		{
@@ -478,7 +487,7 @@ EReimportResult::Type USparseVolumeTextureFactory::Reimport(UObject* Obj)
 #endif // OPENVDB_AVAILABLE
 }
 
-UObject* USparseVolumeTextureFactory::ImportInternal(UClass* InClass, UObject* InParent, FName InName, EObjectFlags Flags, const FString& Filename, const TCHAR* Parms, bool& bOutOperationCanceled)
+UObject* USparseVolumeTextureFactory::ImportInternal(UClass* InClass, UObject* InParent, FName InName, EObjectFlags Flags, const FString& Filename, const TCHAR* Parms, bool& bOutOperationCanceled, bool bIsReimport)
 {
 #if OPENVDB_AVAILABLE
 
@@ -785,7 +794,35 @@ UObject* USparseVolumeTextureFactory::ImportInternal(UClass* InClass, UObject* I
 			return nullptr;
 		}
 
-		UAnimatedSparseVolumeTexture* AnimatedSVTexture = NewObject<UAnimatedSparseVolumeTexture>(InParent, UAnimatedSparseVolumeTexture::StaticClass(), InName, Flags);
+		// By default, the resulting package (already created) and object (about to be) will have the name of the imported file.
+		// However, since the file is part of an entire sequence that we imported, it would be confusing if the resulting asset was named "file_0000" instead of "file",
+		// so we attempt to rename both package and object here.
+		// Don't try to rename the SVT if we are doing a reimport.
+		FName NewObjectName = InName;
+		if (!bIsReimport)
+		{
+			FString NewFileName = GetVDBSequenceBaseFileName(Filename, true /*bDiscardUnderscore*/);
+
+			// Don't try to rename the package if it's the transient package
+			if (InParent != GetTransientPackage() && InParent->IsA<UPackage>())
+			{
+				const FString PackageName = InParent->GetName(); // Contains name with path
+				int32 LastSeparatorIndex = 0;
+				const bool bFoundSeparator = PackageName.FindLastChar(TEXT('/'), LastSeparatorIndex);
+				// Get the substring containing the path only, without the file name
+				const FString PackagePath = bFoundSeparator ? PackageName.Left(LastSeparatorIndex) : FString();
+				const FString NewPackageName = PackagePath / NewFileName;
+
+				UPackage* ExistingPackage = FindPackage(InParent->GetOuter(), *NewPackageName);
+				if (!ExistingPackage)
+				{
+					InParent->Rename(*NewPackageName, nullptr, REN_DontCreateRedirectors | REN_ForceNoResetLoaders);
+					NewObjectName = *NewFileName;
+				}
+			}
+		}
+
+		UAnimatedSparseVolumeTexture* AnimatedSVTexture = NewObject<UAnimatedSparseVolumeTexture>(InParent, UAnimatedSparseVolumeTexture::StaticClass(), NewObjectName, Flags);
 		const bool bInitSuccess = AnimatedSVTexture->Initialize(UncookedFramesData);
 		if (!bInitSuccess)
 		{
diff --git a/Engine/Source/Editor/SparseVolumeTexture/Private/SparseVolumeTextureFactory.h b/Engine/Source/Editor/SparseVolumeTexture/Private/SparseVolumeTextureFactory.h
--- a/Engine/Source/Editor/SparseVolumeTexture/Private/SparseVolumeTextureFactory.h
+++ b/Engine/Source/Editor/SparseVolumeTexture/Private/SparseVolumeTextureFactory.h
@@ -53,7 +53,7 @@ public:
 protected:
 
 private:
-	UObject* ImportInternal(UClass* InClass, UObject* InParent, FName InName, EObjectFlags Flags, const FString& Filename, const TCHAR* Parms, bool& bOutOperationCanceled);
+	UObject* ImportInternal(UClass* InClass, UObject* InParent, FName InName, EObjectFlags Flags, const FString& Filename, const TCHAR* Parms, bool& bOutOperationCanceled, bool bIsReimport);
 };
 
 #endif // WITH_EDITOR
diff --git a/Engine/Source/Runtime/Engine/Classes/SparseVolumeTexture/SparseVolumeTexture.h b/Engine/Source/Runtime/Engine/Classes/SparseVolumeTexture/SparseVolumeTexture.h
--- a/Engine/Source/Runtime/Engine/Classes/SparseVolumeTexture/SparseVolumeTexture.h
+++ b/Engine/Source/Runtime/Engine/Classes/SparseVolumeTexture/SparseVolumeTexture.h
@@ -323,22 +323,25 @@ class UStreamableSparseVolumeTexture : public USparseVolumeTexture
 
 public:
 
-	UPROPERTY(VisibleAnywhere, Category = "Texture")
+	UPROPERTY(VisibleAnywhere, Category = "Texture", AssetRegistrySearchable)
 	FIntVector VolumeResolution;
 
-	UPROPERTY(VisibleAnywhere, Category = "Texture")
+	UPROPERTY(VisibleAnywhere, Category = "Texture", AssetRegistrySearchable)
 	int32 NumMipLevels;
 
-	UPROPERTY(VisibleAnywhere, Category = "Texture")
+	UPROPERTY(VisibleAnywhere, Category = "Texture", AssetRegistrySearchable)
+	int32 NumFrames;
+
+	UPROPERTY(VisibleAnywhere, Category = "Texture", AssetRegistrySearchable)
 	TEnumAsByte<enum EPixelFormat> FormatA;
 
-	UPROPERTY(VisibleAnywhere, Category = "Texture")
+	UPROPERTY(VisibleAnywhere, Category = "Texture", AssetRegistrySearchable)
 	TEnumAsByte<enum EPixelFormat> FormatB;
 
-	UPROPERTY(VisibleAnywhere, Category = "Texture")
+	UPROPERTY(VisibleAnywhere, Category = "Texture", AdvancedDisplay)
 	FVector4f FallbackValueA;
 
-	UPROPERTY(VisibleAnywhere, Category = "Texture")
+	UPROPERTY(VisibleAnywhere, Category = "Texture", AdvancedDisplay)
 	FVector4f FallbackValueB;
 
 	// The addressing mode to use for the X axis.
@@ -354,7 +357,7 @@ public:
 	TEnumAsByte<enum TextureAddress> AddressZ;
 
 	// If enabled, the SparseVolumeTexture is only going to use the local DDC. For certain assets it might be reasonable to also use the remote DDC, but for larger assets this will mean long up- and download times.
-	UPROPERTY(EditAnywhere, BlueprintReadWrite, Category = "Texture", meta = (DisplayName = "Local DDC Only"), AssetRegistrySearchable, AdvancedDisplay)
+	UPROPERTY(EditAnywhere, BlueprintReadWrite, Category = "Texture", AdvancedDisplay)
 	bool bLocalDDCOnly = true;
 
 #if WITH_EDITORONLY_DATA
diff --git a/Engine/Source/Runtime/Engine/Private/SparseVolumeTexture/SparseVolumeTexture.cpp b/Engine/Source/Runtime/Engine/Private/SparseVolumeTexture/SparseVolumeTexture.cpp
--- a/Engine/Source/Runtime/Engine/Private/SparseVolumeTexture/SparseVolumeTexture.cpp
+++ b/Engine/Source/Runtime/Engine/Private/SparseVolumeTexture/SparseVolumeTexture.cpp
@@ -1106,6 +1106,7 @@ bool UStreamableSparseVolumeTexture::EndInitialize(int32 InNumMipLevels)
 	check(NumMipLevelsFullMipChain > 0);
 
 	NumMipLevels = (InNumMipLevels <= INDEX_NONE) ? NumMipLevelsFullMipChain : FMath::Clamp(InNumMipLevels, 1, NumMipLevelsFullMipChain);
+	NumFrames = Frames.Num();
 
 	for (USparseVolumeTextureFrame* Frame : Frames)
 	{
@@ -1163,6 +1164,9 @@ void UStreamableSparseVolumeTexture::PostLoad()
 {
 	Super::PostLoad();
 
+	// Ensure that NumFrames always corresponds to the actual number of frames
+	NumFrames = GetNumFrames();
+
 #if WITH_EDITORONLY_DATA
 	RecacheFrames();
 #else
