SparseVolumeTexture import dialog now defaults to 16bit float for all components and tries to use 8bit unorm for density if possible.

#rnx

diff --git a/Engine/Source/Editor/SparseVolumeTexture/Private/SparseVolumeTextureFactory.cpp b/Engine/Source/Editor/SparseVolumeTexture/Private/SparseVolumeTextureFactory.cpp
--- a/Engine/Source/Editor/SparseVolumeTexture/Private/SparseVolumeTextureFactory.cpp
+++ b/Engine/Source/Editor/SparseVolumeTexture/Private/SparseVolumeTextureFactory.cpp
@@ -42,31 +42,64 @@ static void ComputeDefaultOpenVDBGridAssignment(const TArray<TSharedPtr<FOpenVDB
 			Mapping.SourceGridIndex = INDEX_NONE;
 			Mapping.SourceComponentIndex = INDEX_NONE;
 		}
-		AttributesDesc.Format = ESparseVolumeAttributesFormat::Float32;
+		AttributesDesc.Format = ESparseVolumeAttributesFormat::Float16;
 	}
 
 	// Assign the components of the input grids to the components of the output SVT.
-	uint32 DstAttributesIdx = 0;
-	uint32 DstComponentIdx = 0;
-	for (const TSharedPtr<FOpenVDBGridComponentInfo>& GridComponent : GridComponentInfo)
+	
+	const TSharedPtr<FOpenVDBGridComponentInfo>* DensityComponentInfoPtr = GridComponentInfo.FindByPredicate([](const TSharedPtr<FOpenVDBGridComponentInfo>& GridComponent) { return GridComponent->Name == TEXT("density"); });
+	const int32 NumNonDensityComponentInfos = GridComponentInfo.Num() - 1 - (DensityComponentInfoPtr ? 1 : 0); // -1 because there is always a <None> element in the list
+	
+	// Optimized density assignment: density as 8bit unorm in Attributes A and all other components in Attributes B as 16bit float. This only works if there is a maximum of 4 non-density components.
+	// We also don't use this assignment if there are 3 non-density components as these will be padded to a 4 component format anyways, so we might as well put all 4 components into a single texture.
+	const bool bOptimizedDensityAssignment = (NumNonDensityComponentInfos <= 4 && NumNonDensityComponentInfos != 3) && DensityComponentInfoPtr != nullptr;
+	
+	if (bOptimizedDensityAssignment)
 	{
-		if (GridComponent->Index == INDEX_NONE)
+		// Assign density to the first channel of Attributes A and set format to 8 bit unorm
+		ImportOptions->Attributes[0].Mappings[0].SourceGridIndex = (*DensityComponentInfoPtr)->Index;
+		ImportOptions->Attributes[0].Mappings[0].SourceComponentIndex = (*DensityComponentInfoPtr)->ComponentIndex;
+		ImportOptions->Attributes[0].Format = ESparseVolumeAttributesFormat::Unorm8;
+
+		// All the other components go into Attributes B with 16 bit float
+		uint32 DstComponentIdx = 0;
+		for (const TSharedPtr<FOpenVDBGridComponentInfo>& GridComponent : GridComponentInfo)
 		{
-			continue;
+			if (!ensure(DstComponentIdx <= 3) || GridComponent->Index == INDEX_NONE || &GridComponent == DensityComponentInfoPtr)
+			{
+				continue;
+			}
+			ImportOptions->Attributes[1].Mappings[DstComponentIdx].SourceGridIndex = GridComponent->Index;
+			ImportOptions->Attributes[1].Mappings[DstComponentIdx].SourceComponentIndex = GridComponent->ComponentIndex;
+			++DstComponentIdx;
 		}
-		ImportOptions->Attributes[DstAttributesIdx].Mappings[DstComponentIdx].SourceGridIndex = GridComponent->Index;
-		ImportOptions->Attributes[DstAttributesIdx].Mappings[DstComponentIdx].SourceComponentIndex = GridComponent->ComponentIndex;
-		++DstComponentIdx;
-		if (DstComponentIdx == 4)
+		ImportOptions->Attributes[1].Format = ESparseVolumeAttributesFormat::Float16;
+	}
+	else
+	{
+		uint32 DstAttributesIdx = 0;
+		uint32 DstComponentIdx = 0;
+		for (const TSharedPtr<FOpenVDBGridComponentInfo>& GridComponent : GridComponentInfo)
 		{
-			DstComponentIdx = 0;
-			++DstAttributesIdx;
-			if (DstAttributesIdx == 2)
+			if (GridComponent->Index == INDEX_NONE)
 			{
-				break;
+				continue;
+			}
+			ImportOptions->Attributes[DstAttributesIdx].Mappings[DstComponentIdx].SourceGridIndex = GridComponent->Index;
+			ImportOptions->Attributes[DstAttributesIdx].Mappings[DstComponentIdx].SourceComponentIndex = GridComponent->ComponentIndex;
+			++DstComponentIdx;
+			if (DstComponentIdx == 4)
+			{
+				DstComponentIdx = 0;
+				++DstAttributesIdx;
+				if (DstAttributesIdx == 2)
+				{
+					break;
+				}
 			}
 		}
 	}
+	
 	ImportOptions->bIsSequence = NumFiles > 1;
 }
 
