Fix for SparseVolumeTexture .vdb importer creating broken assets when importing .vdb sequences with this pattern: FILENAME.NUMBER.vdb. When importing a sequence, the importer creates an animated SVT and uses the filename with the number stripped at the end as asset/package name. With files using this pattern, we end up with a dot at the end of the asset/package name, which violates UE naming rules and makes the asset completely unusable without any proper warnings or error messages. The fix now also strips the dot (and other invalid chars).

#rnx

diff --git a/Engine/Source/Editor/SparseVolumeTexture/Private/SparseVolumeTextureFactory.cpp b/Engine/Source/Editor/SparseVolumeTexture/Private/SparseVolumeTextureFactory.cpp
--- a/Engine/Source/Editor/SparseVolumeTexture/Private/SparseVolumeTextureFactory.cpp
+++ b/Engine/Source/Editor/SparseVolumeTexture/Private/SparseVolumeTextureFactory.cpp
@@ -19,6 +19,7 @@
 
 #include "Editor.h"
 #include "EditorFramework/AssetImportData.h"
+#include "ObjectTools.h"
 
 #include "OpenVDBImportWindow.h"
 #include "HAL/PlatformApplicationMisc.h"
@@ -103,18 +104,37 @@ static void ComputeDefaultOpenVDBGridAssignment(const TArray<TSharedPtr<FOpenVDB
 	ImportOptions->bIsSequence = NumFiles > 1;
 }
 
-static FString GetVDBSequenceBaseFileName(const FString& FileName, bool bDiscardUnderscore)
+static FString GetVDBSequenceBaseFileName(const FString& FileName, bool bDiscardNumbersOnly)
 {
 	const FString CleanFileName = FPaths::GetCleanFilename(FileName);
-	const FString CleanFileNameWithoutExt = CleanFileName.LeftChop(4);
-	const int32 LastNonDigitIndex = CleanFileNameWithoutExt.FindLastCharByPredicate([&](TCHAR Letter) 
-		{ 
-			const bool bIsDigit = FChar::IsDigit(Letter);
-			const bool bIsUnderscore = bDiscardUnderscore && Letter == TEXT('_');
-			return !bIsDigit && !bIsUnderscore;
-		}) + 1;
-	const int32 DigitCount = CleanFileNameWithoutExt.Len() - LastNonDigitIndex;
-	const FString CleanFileNameWithoutSuffix = CleanFileNameWithoutExt.LeftChop(DigitCount);
+	int32 NumValidChars = CleanFileName.Len() - 4; // chop off the file extension
+
+	// Remove any digits at the end
+	NumValidChars = CleanFileName.FindLastCharByPredicate([&](TCHAR Letter){ return !FChar::IsDigit(Letter); }, NumValidChars) + 1;
+
+	// Optionally remove other unwanted chars like underscores and invalid object name chars that would later be replaced by underscores
+	if (!bDiscardNumbersOnly)
+	{
+		NumValidChars = CleanFileName.FindLastCharByPredicate([&](TCHAR Letter) 
+			{
+				// INVALID_OBJECTNAME_CHARACTERS is defined in NameTypes.h and is a string literal containing all the invalid chars.
+				// The number at the end of a filename in a sequence is often separated by an underscore or other special character,
+				// so when we want to get the base filename for deriving the new asset name, we also discard these characters.
+				// Underscores are not part of the invalid chars string literal, so we append them here to also discard these chars.
+				const TCHAR* InvalidChars = TEXT("_")INVALID_OBJECTNAME_CHARACTERS;
+				while (*InvalidChars)
+				{
+					if (Letter == *InvalidChars)
+					{
+						return false;
+					}
+					++InvalidChars;
+				}
+				return true;
+			}, NumValidChars) + 1;
+	}
+
+	const FString CleanFileNameWithoutSuffix = CleanFileName.Left(NumValidChars);
 	return CleanFileNameWithoutSuffix;
 }
 
@@ -133,12 +153,24 @@ static TArray<FString> FindOpenVDBSequenceFileNames(const FString& Filename)
 	{
 		const FString Path = FPaths::GetPath(Filename);
 		const FString CleanFilename = FPaths::GetCleanFilename(Filename);
-		const FString CleanFilenameWithoutSuffix = GetVDBSequenceBaseFileName(Filename, false /*bDiscardUnderscore*/);
+		const FString CleanFilenameWithoutSuffix = GetVDBSequenceBaseFileName(Filename, true /*bDiscardNumbersOnly*/);
 
 		// Find all files potentially part of the sequence
 		TArray<FString> PotentialSequenceFilenames;
 		IFileManager::Get().FindFiles(PotentialSequenceFilenames, *Path, TEXT("*.vdb"));
-		PotentialSequenceFilenames = PotentialSequenceFilenames.FilterByPredicate([CleanFilenameWithoutSuffix](const FString& Str) { return Str.StartsWith(CleanFilenameWithoutSuffix); });
+		PotentialSequenceFilenames = PotentialSequenceFilenames.FilterByPredicate([&CleanFilenameWithoutSuffix](const FString& Str) 
+			{ 
+				if (!CleanFilenameWithoutSuffix.IsEmpty())
+				{
+					// Check for the same base filename
+					return Str.StartsWith(CleanFilenameWithoutSuffix);
+				}
+				else
+				{
+					// Removing the digits at the end of the input file resulted in an empty string, so we are looking for numeric filenames only (excluding the 4 chars file extension)
+					return Str.LeftChop(4).IsNumeric();
+				}
+			});
 
 		auto GetFilenameNumberSuffix = [](const FString& Filename) -> int32
 		{
@@ -822,10 +854,13 @@ UObject* USparseVolumeTextureFactory::ImportInternal(UClass* InClass, UObject* I
 		FName NewObjectName = InName;
 		if (!bIsReimport)
 		{
-			FString NewFileName = GetVDBSequenceBaseFileName(Filename, true /*bDiscardUnderscore*/);
+			FString NewFileName = GetVDBSequenceBaseFileName(Filename, false /*bDiscardNumbersOnly*/);
+			// GetVDBSequenceBaseFileName() discards the number as well as underscores and invalid chars at the end of the filename.
+			// We still need to ensure that there are no additional invalid characters in the filename.
+			NewFileName = ObjectTools::SanitizeObjectName(NewFileName);
 
-			// Don't try to rename the package if it's the transient package
-			if (InParent != GetTransientPackage() && InParent->IsA<UPackage>())
+			// Don't try to rename the package if it's the transient package or the new filename is empty
+			if (!NewFileName.IsEmpty() && (InParent != GetTransientPackage() && InParent->IsA<UPackage>()))
 			{
 				const FString PackageName = InParent->GetName(); // Contains name with path
 				int32 LastSeparatorIndex = 0;
