Fix for jittering SparseVolumeTexture frames when importing a sequence with different translations in the transforms of all the input files/frames. A popular third party tool for baking volumetrics unfortunately always creates such sequences. The UnrealEngine .vdb importer assumes that all frames of a sequence share the same transform and effectively ignores that transform. This fix detects changes in the translation and bakes the translation into the created SVT, eliminating the jittering.

#rnx

diff --git a/Engine/Source/Editor/SparseVolumeTexture/Private/SparseVolumeTextureFactory.cpp b/Engine/Source/Editor/SparseVolumeTexture/Private/SparseVolumeTextureFactory.cpp
--- a/Engine/Source/Editor/SparseVolumeTexture/Private/SparseVolumeTextureFactory.cpp
+++ b/Engine/Source/Editor/SparseVolumeTexture/Private/SparseVolumeTextureFactory.cpp
@@ -34,6 +34,14 @@
 
 DEFINE_LOG_CATEGORY_STATIC(LogSparseVolumeTextureFactory, Log, All);
 
+static int32 GSVTImportBakeFrameTranslations = 1;
+static FAutoConsoleVariableRef CVarSVTSImportBakeFrameTranslations(
+	TEXT("r.SparseVolumeTexture.Import.BakeVDBFrameTranslations"),
+	GSVTImportBakeFrameTranslations,
+	TEXT("If enabled, VDB sequences with translations that change across the frames/files will have the translation baked into the resulting SVT frames."),
+	ECVF_Default
+);
+
 static void ComputeDefaultOpenVDBGridAssignment(const TArray<TSharedPtr<FOpenVDBGridComponentInfo>>& GridComponentInfo, int32 NumFiles, FOpenVDBImportOptions* ImportOptions)
 {
 	for (FOpenVDBSparseVolumeAttributesDesc& AttributesDesc : ImportOptions->Attributes)
@@ -592,7 +600,7 @@ UObject* USparseVolumeTextureFactory::ImportInternal(UClass* InClass, UObject* I
 	}
 
 	// Utility function for computing the bounding box encompassing the bounds of all frames in the SVT.
-	auto ExpandVolumeBounds = [](const FOpenVDBImportOptions& ImportOptions, const TArray<FOpenVDBGridInfo>& GridInfoArray, FIntVector3& VolumeBoundsMin, FIntVector3& VolumeBoundsMax)
+	auto ExpandVolumeBounds = [](const FOpenVDBImportOptions& ImportOptions, const TArray<FOpenVDBGridInfo>& GridInfoArray, FIntVector3& VolumeBoundsMin, FIntVector3& VolumeBoundsMax, float* TransformMin, float* TransformMax)
 	{
 		for (const FOpenVDBSparseVolumeAttributesDesc& Attributes : ImportOptions.Attributes)
 		{
@@ -608,6 +616,15 @@ UObject* USparseVolumeTextureFactory::ImportInternal(UClass* InClass, UObject* I
 					VolumeBoundsMax.X = FMath::Max(VolumeBoundsMax.X, GridInfo.VolumeActiveAABBMax.X);
 					VolumeBoundsMax.Y = FMath::Max(VolumeBoundsMax.Y, GridInfo.VolumeActiveAABBMax.Y);
 					VolumeBoundsMax.Z = FMath::Max(VolumeBoundsMax.Z, GridInfo.VolumeActiveAABBMax.Z);
+
+					for (int32 y = 0; y < 4; ++y)
+					{
+						for (int32 x = 0; x < 4; ++x)
+						{
+							TransformMin[y * 4 + x] = FMath::Min(TransformMin[y * 4 + x], GridInfo.Transform.M[y][x]);
+							TransformMax[y * 4 + x] = FMath::Max(TransformMax[y * 4 + x], GridInfo.Transform.M[y][x]);
+						}
+					}
 				}
 			}
 		}
@@ -615,6 +632,13 @@ UObject* USparseVolumeTextureFactory::ImportInternal(UClass* InClass, UObject* I
 
 	FIntVector3 VolumeBoundsMin = FIntVector3(INT32_MAX, INT32_MAX, INT32_MAX);
 	FIntVector3 VolumeBoundsMax = FIntVector3(INT32_MIN, INT32_MIN, INT32_MIN);
+	float TransformMin[16];
+	float TransformMax[16];
+	for (int32 i = 0; i < 16; ++i)
+	{
+		TransformMin[i] = FLT_MAX;
+		TransformMax[i] = -FLT_MAX;
+	}
 
 	// Import as either single static SVT or a sequence of frames, making up an animated SVT
 	if (!ImportOptions.bIsSequence)
@@ -624,10 +648,10 @@ UObject* USparseVolumeTextureFactory::ImportInternal(UClass* InClass, UObject* I
 		FScopedSlowTask ImportTask(1.0f, LOCTEXT("ImportingVDBStatic", "Importing static OpenVDB"));
 		ImportTask.MakeDialog(true);
 
-		ExpandVolumeBounds(ImportOptions, PreviewData.GridInfo, VolumeBoundsMin, VolumeBoundsMax);
+		ExpandVolumeBounds(ImportOptions, PreviewData.GridInfo, VolumeBoundsMin, VolumeBoundsMax, TransformMin, TransformMax);
 
 		UE::SVT::FTextureData TextureData{};
-		const bool bConversionSuccess = ConvertOpenVDBToSparseVolumeTexture(PreviewData.LoadedFile, ImportOptions, VolumeBoundsMin, TextureData);
+		const bool bConversionSuccess = ConvertOpenVDBToSparseVolumeTexture(PreviewData.LoadedFile, ImportOptions, VolumeBoundsMin, false /*bBakeTranslation*/, TextureData);
 
 		if (!bConversionSuccess)
 		{
@@ -673,7 +697,7 @@ UObject* USparseVolumeTextureFactory::ImportInternal(UClass* InClass, UObject* I
 
 		// Compute volume bounds and check sequence files for compatiblity
 		std::mutex VolumeBoundsMutex;
-		ParallelFor(NumFrames, [&bErrored, &VolumeBoundsMutex, &VolumeBoundsMin, &VolumeBoundsMax, &ExpandVolumeBounds, &PreviewData, &ImportOptions](int32 FrameIdx)
+		ParallelFor(NumFrames, [&bErrored, &VolumeBoundsMutex, &VolumeBoundsMin, &VolumeBoundsMax, TransformMin = &TransformMin[0], TransformMax = &TransformMax[0], &ExpandVolumeBounds, &PreviewData, &ImportOptions](int32 FrameIdx)
 			{
 				if (bErrored.load())
 				{
@@ -727,7 +751,7 @@ UObject* USparseVolumeTextureFactory::ImportInternal(UClass* InClass, UObject* I
 				// Update sequence volume bounds and increment ProcessedFramesCounter
 				{
 					std::lock_guard<std::mutex> Lock(VolumeBoundsMutex);
-					ExpandVolumeBounds(ImportOptions, FrameGridInfo, VolumeBoundsMin, VolumeBoundsMax);
+					ExpandVolumeBounds(ImportOptions, FrameGridInfo, VolumeBoundsMin, VolumeBoundsMax, TransformMin, TransformMax);
 				}
 			});
 
@@ -736,6 +760,39 @@ UObject* USparseVolumeTextureFactory::ImportInternal(UClass* InClass, UObject* I
 			return nullptr;
 		}
 
+		FMatrix44f TransformMatrixMin;
+		FMatrix44f TransformMatrixMax;
+		for (int32 y = 0; y < 4; ++y)
+		{
+			for (int32 x = 0; x < 4; ++x)
+			{
+				TransformMatrixMin.M[y][x] = TransformMin[y * 4 + x];
+				TransformMatrixMax.M[y][x] = TransformMax[y * 4 + x];
+			}
+		}
+		const FVector3f TransformScaleMin = TransformMatrixMin.GetScaleVector();
+		const FVector3f TransformScaleMax = TransformMatrixMax.GetScaleVector();
+		const FVector3f TransformTransMin = TransformMatrixMin.GetOrigin();
+		const FVector3f TransformTransMax = TransformMatrixMax.GetOrigin();
+		const bool bSameScaleForAllFrames = (TransformScaleMax - TransformScaleMin).GetAbsMax() <= UE_SMALL_NUMBER;
+		const bool bSameTranslationForAllFrames = (TransformTransMax - TransformTransMin).GetAbsMax() <= UE_SMALL_NUMBER;
+
+		// We currently do not support changing scale/rotation across frames
+		if (!bSameScaleForAllFrames)
+		{
+			UE_LOG(LogSparseVolumeTextureFactory, Warning, TEXT("OpenVDB sequence is using different scalings and rotations across all frames"));
+		}
+
+		const bool bBakeTranslations = !bSameTranslationForAllFrames && (GSVTImportBakeFrameTranslations != 0);
+
+		// But we can bake translation changes into the SVT
+		if (bBakeTranslations)
+		{
+			FVector3f OffsetF = TransformTransMin / TransformScaleMin;
+			FIntVector3 Offset = FIntVector3(FMath::Floor(OffsetF.X), FMath::Floor(OffsetF.Y), FMath::Floor(OffsetF.Z));
+			VolumeBoundsMin += Offset;
+		}
+
 		ImportTask.EnterProgressFrame(1.0f, LOCTEXT("ConvertingVDBAnim", "Converting OpenVDB animation"));
 
 		FEvent* AllTasksFinishedEvent = FPlatformProcess::GetSynchEventFromPool();
@@ -764,7 +821,7 @@ UObject* USparseVolumeTextureFactory::ImportInternal(UClass* InClass, UObject* I
 
 			AsyncTask(ENamedThreads::AnyNormalThreadNormalTask,
 				[FrameIdx, NumFrames, &PreviewData, &ImportOptions, &UncookedFramesData,
-				AllTasksFinishedEvent, &bErrored, &bCanceled, &FinishedTasksCounter, &ProcessedFramesCounter, &VolumeBoundsMin]()
+				AllTasksFinishedEvent, &bErrored, &bCanceled, &FinishedTasksCounter, &ProcessedFramesCounter, &VolumeBoundsMin, bBakeTranslations]()
 				{
 					// Ensure the FinishedTasksCounter will be incremented in all cases
 					FScopedIncrementer Incremeter(FinishedTasksCounter, NumFrames, AllTasksFinishedEvent);
@@ -788,7 +845,7 @@ UObject* USparseVolumeTextureFactory::ImportInternal(UClass* InClass, UObject* I
 					}
 
 					UE::SVT::FTextureData TextureData{};
-					const bool bConversionSuccess = ConvertOpenVDBToSparseVolumeTexture(LoadedFrameFile, ImportOptions, VolumeBoundsMin, TextureData);
+					const bool bConversionSuccess = ConvertOpenVDBToSparseVolumeTexture(LoadedFrameFile, ImportOptions, VolumeBoundsMin, bBakeTranslations, TextureData);
 
 					if (!bConversionSuccess)
 					{
diff --git a/Engine/Source/Editor/SparseVolumeTexture/Private/SparseVolumeTextureOpenVDBUtility.cpp b/Engine/Source/Editor/SparseVolumeTexture/Private/SparseVolumeTextureOpenVDBUtility.cpp
--- a/Engine/Source/Editor/SparseVolumeTexture/Private/SparseVolumeTextureOpenVDBUtility.cpp
+++ b/Engine/Source/Editor/SparseVolumeTexture/Private/SparseVolumeTextureOpenVDBUtility.cpp
@@ -111,7 +111,7 @@ static FOpenVDBGridInfo GetOpenVDBGridInfo(openvdb::GridBase::Ptr Grid, uint32 G
 	{
 		for (int j = 0; j < 4; ++j)
 		{
-			GridInfo.Transform.M[i][j] = static_cast<float>(GridTransformVDB[j][i]);
+			GridInfo.Transform.M[i][j] = static_cast<float>(GridTransformVDB[i][j]);
 		}
 	}
 
@@ -285,10 +285,11 @@ class FSparseVolumeTextureDataProviderOpenVDB : public UE::SVT::ITextureDataProv
 {
 public:
 
-	bool Initialize(TArray64<uint8>& SourceFile, const FOpenVDBImportOptions& ImportOptions, const FIntVector3& InVolumeBoundsMin)
+	bool Initialize(TArray64<uint8>& SourceFile, const FOpenVDBImportOptions& ImportOptions, const FIntVector3& InVolumeBoundsMin, bool bInBakeTranslation)
 	{
 		Attributes = ImportOptions.Attributes;
 		VolumeBoundsMin = InVolumeBoundsMin;
+		Translation = FIntVector3::ZeroValue;
 
 		// Compute some basic info about the number of components and which format to use
 		EPixelFormat MultiCompFormat[NumAttributesDescs] = {};
@@ -366,6 +367,10 @@ public:
 
 		FIntVector3 SmallestAABBMin = FIntVector3(INT32_MAX);
 		FIntVector3 LargestAABBMax = FIntVector3(INT32_MIN);
+		FVector3f ScaleMin = FVector3f(FLT_MAX);
+		FVector3f ScaleMax = FVector3f(-FLT_MAX);
+		FVector3f TranslationMin = FVector3f(FLT_MAX);
+		FVector3f TranslationMax = FVector3f(-FLT_MAX);
 
 		// Compute per source grid data of up to 4 different grids (one per component)
 		UniqueGridAdapters.SetNum((int32)NumSourceGrids);
@@ -407,6 +412,13 @@ public:
 				LargestAABBMax.Y = FMath::Max(LargestAABBMax.Y, GridInfo.VolumeActiveAABBMax.Y);
 				LargestAABBMax.Z = FMath::Max(LargestAABBMax.Z, GridInfo.VolumeActiveAABBMax.Z);
 
+				const FVector3f GridScale = GridInfo.Transform.GetScaleVector();
+				const FVector3f GridTranslation = GridInfo.Transform.GetOrigin();
+				ScaleMin = FVector3f::Min(ScaleMin, GridScale);
+				ScaleMax = FVector3f::Max(ScaleMax, GridScale);
+				TranslationMin = FVector3f::Min(TranslationMin, GridTranslation);
+				TranslationMax = FVector3f::Max(TranslationMax, GridTranslation);
+
 				SVTCreateInfo.FallbackValues[AttributesIdx][CompIdx] = UniqueGridAdapters[SourceGridIndex]->GetBackgroundValue(SourceComponentIndex);
 
 				FSingleGridToComponentMapping Mapping{};
@@ -416,7 +428,24 @@ public:
 				GridToComponentMappings[SourceGridIndex].Add(Mapping);
 			}
 		}
+
+		const float ScaleDiffAbs = FVector3f(ScaleMax - ScaleMin).GetAbsMax();
+		const float TransDiffAbs = FVector3f(TranslationMax - TranslationMin).GetAbsMax();
+		if (bInBakeTranslation && ((ScaleDiffAbs > UE_SMALL_NUMBER) || (TransDiffAbs > UE_SMALL_NUMBER)))
+		{
+			UE_LOG(LogSparseVolumeTextureOpenVDBUtility, Warning, TEXT("Frame has multiple grids with different transforms!"));
+		}
+
 		const bool bEmptyBounds = SmallestAABBMin.X >= LargestAABBMax.X || SmallestAABBMin.Y >= LargestAABBMax.Y || SmallestAABBMin.Z >= LargestAABBMax.Z;
+
+		if (bInBakeTranslation && !bEmptyBounds)
+		{
+			FVector3f OffsetF = TranslationMin / ScaleMin;
+			Translation = FIntVector3(FMath::Floor(OffsetF.X), FMath::Floor(OffsetF.Y), FMath::Floor(OffsetF.Z));
+			SmallestAABBMin += Translation;
+			LargestAABBMax += Translation;
+		}
+
 		SVTCreateInfo.VirtualVolumeAABBMin = bEmptyBounds ? FIntVector3(INT32_MAX) : (SmallestAABBMin - VolumeBoundsMin);
 		SVTCreateInfo.VirtualVolumeAABBMax = bEmptyBounds ? FIntVector3(INT32_MIN) : (LargestAABBMax - VolumeBoundsMin);
 
@@ -440,9 +469,10 @@ public:
 			UniqueGridAdapters[GridIdx]->IteratePhysical(
 				[&](const FIntVector3& Coord, uint32 NumVoxelComponents, float* VoxelValues)
 				{
+					FIntVector3 RemappedCoord = Coord + Translation - VolumeBoundsMin;
 					for (const FSingleGridToComponentMapping& Mapping : GridToComponentMappings[GridIdx])
 					{
-						OnVisit(Coord - VolumeBoundsMin, Mapping.AttributesIdx, Mapping.ComponentIdx, VoxelValues[Mapping.GridComponentIdx]);
+						OnVisit(RemappedCoord, Mapping.AttributesIdx, Mapping.ComponentIdx, VoxelValues[Mapping.GridComponentIdx]);
 					}
 				});
 		}
@@ -464,15 +494,16 @@ private:
 	TStaticArray<uint32, NumAttributesDescs> NumComponents;
 	UE::SVT::FTextureDataCreateInfo SVTCreateInfo;
 	FIntVector3 VolumeBoundsMin;
+	FIntVector3 Translation;
 };
 
 #endif // OPENVDB_AVAILABLE
 
-bool ConvertOpenVDBToSparseVolumeTexture(TArray64<uint8>& SourceFile, const FOpenVDBImportOptions& ImportOptions, const FIntVector3& VolumeBoundsMin, UE::SVT::FTextureData& OutResult)
+bool ConvertOpenVDBToSparseVolumeTexture(TArray64<uint8>& SourceFile, const FOpenVDBImportOptions& ImportOptions, const FIntVector3& VolumeBoundsMin, bool bBakeTranslation, UE::SVT::FTextureData& OutResult)
 {
 #if OPENVDB_AVAILABLE
 	FSparseVolumeTextureDataProviderOpenVDB DataProvider;
-	if (!DataProvider.Initialize(SourceFile, ImportOptions, VolumeBoundsMin))
+	if (!DataProvider.Initialize(SourceFile, ImportOptions, VolumeBoundsMin, bBakeTranslation))
 	{
 		return false;
 	}
diff --git a/Engine/Source/Editor/SparseVolumeTexture/Private/SparseVolumeTextureOpenVDBUtility.h b/Engine/Source/Editor/SparseVolumeTexture/Private/SparseVolumeTextureOpenVDBUtility.h
--- a/Engine/Source/Editor/SparseVolumeTexture/Private/SparseVolumeTextureOpenVDBUtility.h
+++ b/Engine/Source/Editor/SparseVolumeTexture/Private/SparseVolumeTextureOpenVDBUtility.h
@@ -59,7 +59,7 @@ bool IsOpenVDBGridValid(const FOpenVDBGridInfo& GridInfo, const FString& Filenam
 
 bool GetOpenVDBGridInfo(TArray64<uint8>& SourceFile, bool bCreateStrings, TArray<FOpenVDBGridInfo>* OutGridInfo);
 
-bool ConvertOpenVDBToSparseVolumeTexture(TArray64<uint8>& SourceFile, const struct FOpenVDBImportOptions& ImportOptions, const FIntVector3& VolumeBoundsMin, UE::SVT::FTextureData& OutResult);
+bool ConvertOpenVDBToSparseVolumeTexture(TArray64<uint8>& SourceFile, const struct FOpenVDBImportOptions& ImportOptions, const FIntVector3& VolumeBoundsMin, bool bBakeTranslation, UE::SVT::FTextureData& OutResult);
 
 const TCHAR* OpenVDBGridTypeToString(EOpenVDBGridType Type);
 
