SVT: Small reshuffle to allow C++ modules to set import options programatically.

diff --git a/Engine/Source/Editor/SparseVolumeTexture/Private/OpenVDBImportOptions.h b/Engine/Source/Editor/SparseVolumeTexture/Private/OpenVDBImportOptions.h
deleted file mode 100644
diff --git a/Engine/Source/Editor/SparseVolumeTexture/Private/OpenVDBImportWindow.h b/Engine/Source/Editor/SparseVolumeTexture/Private/OpenVDBImportWindow.h
--- a/Engine/Source/Editor/SparseVolumeTexture/Private/OpenVDBImportWindow.h
+++ b/Engine/Source/Editor/SparseVolumeTexture/Private/OpenVDBImportWindow.h
@@ -11,17 +11,10 @@
 
 struct FOpenVDBGridInfo;
 struct FOpenVDBImportOptions;
+struct FOpenVDBGridComponentInfo;
 struct FOpenVDBSparseVolumeAttributesDesc;
 enum class ESparseVolumeAttributesFormat : uint8;
 
-struct FOpenVDBGridComponentInfo
-{
-	uint32 Index;
-	uint32 ComponentIndex;
-	FString Name;
-	FString DisplayString; // Contains source file grid index, name and component (if it is a multi component type like Float3)
-};
-
 class SOpenVDBGridInfoTableRow : public SMultiColumnTableRow<TSharedPtr<FOpenVDBGridInfo>>
 {
 public:
diff --git a/Engine/Source/Editor/SparseVolumeTexture/Private/SparseVolumeTextureFactory.cpp b/Engine/Source/Editor/SparseVolumeTexture/Private/SparseVolumeTextureFactory.cpp
--- a/Engine/Source/Editor/SparseVolumeTexture/Private/SparseVolumeTextureFactory.cpp
+++ b/Engine/Source/Editor/SparseVolumeTexture/Private/SparseVolumeTextureFactory.cpp
@@ -17,6 +17,7 @@
 #include "Async/Async.h"
 #include "Async/ParallelFor.h"
 
+#include "AssetImportTask.h"
 #include "Editor.h"
 #include "EditorFramework/AssetImportData.h"
 #include "ObjectTools.h"
@@ -239,17 +240,7 @@ static TArray<FString> FindOpenVDBSequenceFileNames(const FString& Filename)
 	return SequenceFilenames;
 }
 
-struct FOpenVDBPreviewData
-{
-	TArray64<uint8> LoadedFile;
-	TArray<FOpenVDBGridInfo> GridInfo;
-	TArray<TSharedPtr<FOpenVDBGridInfo>> GridInfoPtrs;
-	TArray<TSharedPtr<FOpenVDBGridComponentInfo>> GridComponentInfoPtrs;
-	TArray<FString> SequenceFilenames;
-	FOpenVDBImportOptions DefaultImportOptions;
-};
-
-static bool LoadOpenVDBPreviewData(const FString& Filename, FOpenVDBPreviewData* OutPreviewData)
+bool LoadOpenVDBPreviewData(const FString& Filename, FOpenVDBPreviewData* OutPreviewData)
 {
 	FOpenVDBPreviewData& Result = *OutPreviewData;
 	check(Result.LoadedFile.IsEmpty());
@@ -574,11 +565,26 @@ UObject* USparseVolumeTextureFactory::ImportInternal(UClass* InClass, UObject* I
 		|| IsRunningCommandlet()
 		|| GIsRunningUnattendedScript);
 
-	// Load file and get info about each contained grid
 	FOpenVDBPreviewData PreviewData;
-	if (!LoadOpenVDBPreviewData(Filename, &PreviewData))
+
+	// Use the provided preview data, if any
+	bool bCollectedData = false;
+	if (bIsUnattended && AssetImportTask)
 	{
-		return nullptr;
+		if (UOpenVDBImportOptionsObject* TaskOptions = Cast<UOpenVDBImportOptionsObject>(AssetImportTask->Options))
+		{
+			PreviewData = TaskOptions->PreviewData;
+			bCollectedData = true;
+		}
+	}
+
+	// Otherwise, load file and get info about each contained grid
+	if (!bCollectedData)
+	{
+		if (!LoadOpenVDBPreviewData(Filename, &PreviewData))
+		{
+			return nullptr;
+		}
 	}
 
 	FOpenVDBImportOptions ImportOptions = PreviewData.DefaultImportOptions;
diff --git a/Engine/Source/Editor/SparseVolumeTexture/Private/SparseVolumeTextureOpenVDBUtility.h b/Engine/Source/Editor/SparseVolumeTexture/Private/SparseVolumeTextureOpenVDBUtility.h
--- a/Engine/Source/Editor/SparseVolumeTexture/Private/SparseVolumeTextureOpenVDBUtility.h
+++ b/Engine/Source/Editor/SparseVolumeTexture/Private/SparseVolumeTextureOpenVDBUtility.h
@@ -13,39 +13,8 @@ namespace UE
 		struct FTextureData;
 	}
 }
-
-enum class EOpenVDBGridType : uint8
-{
-	Unknown = 0,
-	Half,
-	Half2,
-	Half3,
-	Half4,
-	Float,
-	Float2,
-	Float3,
-	Float4,
-	Double,
-	Double2,
-	Double3,
-	Double4,
-};
-
-struct FOpenVDBGridInfo
-{
-	FMatrix44f Transform;
-	FIntVector3 VolumeActiveAABBMin;
-	FIntVector3 VolumeActiveAABBMax;
-	FIntVector3 VolumeActiveDim;
-	FVector VolumeVoxelSize;
-	FString Name;
-	FString DisplayString; // Contains Index (into source file grids), Type and Name
-	uint32 Index;
-	uint32 NumComponents;
-	EOpenVDBGridType Type;
-	bool bIsInWorldSpace;
-	bool bHasUniformVoxels;
-};
+struct FOpenVDBGridInfo;
+enum class EOpenVDBGridType : uint8;
 
 struct FOpenVDBToSVTConversionResult
 {
diff --git a/Engine/Source/Editor/SparseVolumeTexture/Public/OpenVDBImportOptions.h b/Engine/Source/Editor/SparseVolumeTexture/Public/OpenVDBImportOptions.h
new file mode 100644
--- /dev/null
+++ b/Engine/Source/Editor/SparseVolumeTexture/Public/OpenVDBImportOptions.h
@@ -0,0 +1,93 @@
+// Copyright Epic Games, Inc. All Rights Reserved.
+
+#pragma once
+
+#include "CoreMinimal.h"
+#include "Containers/StaticArray.h"
+
+#include "OpenVDBImportOptions.generated.h"
+
+enum class ESparseVolumeAttributesFormat : uint8
+{
+	Unorm8 = 0,
+	Float16 = 1,
+	Float32 = 2,
+};
+
+struct FOpenVDBSparseVolumeComponentMapping
+{
+	int32 SourceGridIndex = INDEX_NONE;
+	int32 SourceComponentIndex = INDEX_NONE;
+};
+
+struct FOpenVDBSparseVolumeAttributesDesc
+{
+	TStaticArray<FOpenVDBSparseVolumeComponentMapping, 4> Mappings;
+	ESparseVolumeAttributesFormat Format = ESparseVolumeAttributesFormat::Unorm8;
+};
+
+struct FOpenVDBImportOptions
+{
+	TStaticArray<FOpenVDBSparseVolumeAttributesDesc, 2> Attributes;
+	bool bIsSequence = false;
+};
+
+enum class EOpenVDBGridType : uint8
+{
+	Unknown = 0,
+	Half,
+	Half2,
+	Half3,
+	Half4,
+	Float,
+	Float2,
+	Float3,
+	Float4,
+	Double,
+	Double2,
+	Double3,
+	Double4,
+};
+
+struct FOpenVDBGridInfo
+{
+	FMatrix44f Transform;
+	FIntVector3 VolumeActiveAABBMin;
+	FIntVector3 VolumeActiveAABBMax;
+	FIntVector3 VolumeActiveDim;
+	FVector VolumeVoxelSize;
+	FString Name;
+	FString DisplayString; // Contains Index (into source file grids), Type and Name
+	uint32 Index;
+	uint32 NumComponents;
+	EOpenVDBGridType Type;
+	bool bIsInWorldSpace;
+	bool bHasUniformVoxels;
+};
+
+struct FOpenVDBGridComponentInfo
+{
+	uint32 Index;
+	uint32 ComponentIndex;
+	FString Name;
+	FString DisplayString; // Contains source file grid index, name and component (if it is a multi component type like Float3)
+};
+
+struct FOpenVDBPreviewData
+{
+	TArray64<uint8> LoadedFile;
+	TArray<FOpenVDBGridInfo> GridInfo;
+	TArray<TSharedPtr<FOpenVDBGridInfo>> GridInfoPtrs;
+	TArray<TSharedPtr<FOpenVDBGridComponentInfo>> GridComponentInfoPtrs;
+	TArray<FString> SequenceFilenames;
+	FOpenVDBImportOptions DefaultImportOptions;
+};
+
+UCLASS()
+class SPARSEVOLUMETEXTURE_API UOpenVDBImportOptionsObject : public UObject
+{
+	GENERATED_BODY()
+
+public:
+	FOpenVDBPreviewData PreviewData;
+};
\ No newline at end of file
diff --git a/Engine/Source/Editor/SparseVolumeTexture/Public/SparseVolumeTextureFactory.h b/Engine/Source/Editor/SparseVolumeTexture/Public/SparseVolumeTextureFactory.h
--- a/Engine/Source/Editor/SparseVolumeTexture/Public/SparseVolumeTextureFactory.h
+++ b/Engine/Source/Editor/SparseVolumeTexture/Public/SparseVolumeTextureFactory.h
@@ -10,6 +10,8 @@
 
 #include "SparseVolumeTextureFactory.generated.h"
 
+struct FOpenVDBPreviewData;
+
 // Responsible for creating and importing Sparse Volume Texture objects
 UCLASS(hidecategories = Object, MinimalAPI)
 class USparseVolumeTextureFactory : public UFactory, public FReimportHandler
@@ -56,4 +58,6 @@ private:
 	UObject* ImportInternal(UClass* InClass, UObject* InParent, FName InName, EObjectFlags Flags, const FString& Filename, const TCHAR* Parms, bool& bOutOperationCanceled, bool bIsReimport);
 };
 
+bool SPARSEVOLUMETEXTURE_API LoadOpenVDBPreviewData(const FString& Filename, FOpenVDBPreviewData* OutPreviewData);
+
 #endif // WITH_EDITOR
