Implemented support for SVT per-frame transforms. Transforms are read from from the source vdb files and are required to be identical for all grids in the same frame. The SVT viewer now implements support for applying these per-frame transforms. It does so by applying the the inverse per-frame transform when transforming from (translated )world space to local unit (and eventually UV-) space.
Fixed SVT viewer breaking at LWC distances. It now uses translated world space instead.

#rnx

diff --git a/Engine/Plugins/FX/Niagara/Source/Niagara/Private/NiagaraDataInterfaceRenderTargetVolume.cpp b/Engine/Plugins/FX/Niagara/Source/Niagara/Private/NiagaraDataInterfaceRenderTargetVolume.cpp
--- a/Engine/Plugins/FX/Niagara/Source/Niagara/Private/NiagaraDataInterfaceRenderTargetVolume.cpp
+++ b/Engine/Plugins/FX/Niagara/Source/Niagara/Private/NiagaraDataInterfaceRenderTargetVolume.cpp
@@ -692,7 +692,7 @@ bool UNiagaraDataInterfaceRenderTargetVolume::SimCacheWriteFrame(UObject* Storag
 					return false;
 				}
 
-				CurrCache->AppendFrame(SparseTextureData);
+				CurrCache->AppendFrame(SparseTextureData, FTransform::Identity);
 #endif
 			}
 			else
diff --git a/Engine/Plugins/FX/Niagara/Source/NiagaraEditor/Private/NiagaraBakerRendererOutputSparseVolumeTexture.cpp b/Engine/Plugins/FX/Niagara/Source/NiagaraEditor/Private/NiagaraBakerRendererOutputSparseVolumeTexture.cpp
--- a/Engine/Plugins/FX/Niagara/Source/NiagaraEditor/Private/NiagaraBakerRendererOutputSparseVolumeTexture.cpp
+++ b/Engine/Plugins/FX/Niagara/Source/NiagaraEditor/Private/NiagaraBakerRendererOutputSparseVolumeTexture.cpp
@@ -302,7 +302,7 @@ void FNiagaraBakerRendererOutputSparseVolumeTexture::BakeFrame(FNiagaraBakerFeed
 			return;
 		}
 
-		if (!SVTAsset->AppendFrame(SparseTextureData))
+		if (!SVTAsset->AppendFrame(SparseTextureData, FTransform::Identity))
 		{
 			UE_LOG(LogNiagaraBaker, Error, TEXT("Cannot append frame to SVT"));
 		}
diff --git a/Engine/Shaders/Private/SparseVolumeTexture/VisualizeSparseVolumeTexture.usf b/Engine/Shaders/Private/SparseVolumeTexture/VisualizeSparseVolumeTexture.usf
--- a/Engine/Shaders/Private/SparseVolumeTexture/VisualizeSparseVolumeTexture.usf
+++ b/Engine/Shaders/Private/SparseVolumeTexture/VisualizeSparseVolumeTexture.usf
@@ -52,9 +52,9 @@ int MipLevel;
 float4 WorldToLocal0;
 float4 WorldToLocal1;
 float4 WorldToLocal2;
-float3 WorldToLocalNoScale0;
-float3 WorldToLocalNoScale1;
-float3 WorldToLocalNoScale2;
+float3 WorldToLocalRotation0;
+float3 WorldToLocalRotation1;
+float3 WorldToLocalRotation2;
 uint ComponentToVisualize;
 float Extinction;
 
@@ -64,15 +64,10 @@ void VisualizeSparseVolumeTexturePS(
 )
 {
 	ResolvedView = ResolveView();
-	const uint2 PixelCoord = uint2(SVPos.xy);
-	float2 UvBuffer = PixelCoord * View.BufferSizeAndInvSize.zw;	// Uv for depth buffer read (size can be larger than viewport)
 
-	float2 ScreenPosition = SvPositionToScreenPosition(SVPos).xy;
-	const float Depth = 10000.0f;
-	float4 TranslatedWorldPos = mul(float4(ScreenPosition * Depth, Depth, 1), View.ScreenToTranslatedWorld);
-
-	const float3 RayDirWorld = normalize(TranslatedWorldPos.xyz - View.TranslatedWorldCameraOrigin);
-	const float3 RayOrigWorld = DFHackToFloat(ResolvedView.WorldCameraOrigin);
+	const float3 TranslatedWorldPos = SvPositionToResolvedTranslatedWorld(SVPos);
+	const float3 RayDirWorld = GetCameraVectorFromTranslatedWorldPosition(TranslatedWorldPos);
+	const float3 RayOriginWorld = SvPositionToResolvedTranslatedWorld(float4(SVPos.xy, NearDepthValue, 1.0f));
 
 	// Transform into local [-1.0f, 1.0f] space
 	float3 RayDir;
@@ -81,9 +76,9 @@ void VisualizeSparseVolumeTexturePS(
 	RayDir.z = dot(WorldToLocal2.xyz, RayDirWorld);
 
 	float3 RayOrig;
-	RayOrig.x = dot(WorldToLocal0, float4(RayOrigWorld, 1.0f));
-	RayOrig.y = dot(WorldToLocal1, float4(RayOrigWorld, 1.0f));
-	RayOrig.z = dot(WorldToLocal2, float4(RayOrigWorld, 1.0f));
+	RayOrig.x = dot(WorldToLocal0, float4(RayOriginWorld, 1.0f));
+	RayOrig.y = dot(WorldToLocal1, float4(RayOriginWorld, 1.0f));
+	RayOrig.z = dot(WorldToLocal2, float4(RayOriginWorld, 1.0f));
 
 	OutLuminanceAlpha = float4(0.0, 0.0, 0.0, 1.0);
 
@@ -115,9 +110,9 @@ void VisualizeSparseVolumeTexturePS(
 
 		// DDA requires the ray direction not to be non-(uniformly) scaled.
 		float3 RayDirNoScale;
-		RayDirNoScale.x = dot(WorldToLocalNoScale0, RayDirWorld);
-		RayDirNoScale.y = dot(WorldToLocalNoScale1, RayDirWorld);
-		RayDirNoScale.z = dot(WorldToLocalNoScale2, RayDirWorld);
+		RayDirNoScale.x = dot(WorldToLocalRotation0, RayDirWorld);
+		RayDirNoScale.y = dot(WorldToLocalRotation1, RayDirWorld);
+		RayDirNoScale.z = dot(WorldToLocalRotation2, RayDirWorld);
 
 		// Amanatides 3D-DDA data preparation
 		float3 stepSign = sign(RayDirNoScale);
diff --git a/Engine/Source/Editor/SparseVolumeTexture/Private/SparseVolumeTextureFactory.cpp b/Engine/Source/Editor/SparseVolumeTexture/Private/SparseVolumeTextureFactory.cpp
--- a/Engine/Source/Editor/SparseVolumeTexture/Private/SparseVolumeTextureFactory.cpp
+++ b/Engine/Source/Editor/SparseVolumeTexture/Private/SparseVolumeTextureFactory.cpp
@@ -35,14 +35,6 @@
 
 DEFINE_LOG_CATEGORY_STATIC(LogSparseVolumeTextureFactory, Log, All);
 
-static int32 GSVTImportBakeFrameTranslations = 1;
-static FAutoConsoleVariableRef CVarSVTSImportBakeFrameTranslations(
-	TEXT("r.SparseVolumeTexture.Import.BakeVDBFrameTranslations"),
-	GSVTImportBakeFrameTranslations,
-	TEXT("If enabled, VDB sequences with translations that change across the frames/files will have the translation baked into the resulting SVT frames."),
-	ECVF_Default
-);
-
 static void ComputeDefaultOpenVDBGridAssignment(const TArray<TSharedPtr<FOpenVDBGridComponentInfo>>& GridComponentInfo, int32 NumFiles, FOpenVDBImportOptions* ImportOptions)
 {
 	for (FOpenVDBSparseVolumeAttributesDesc& AttributesDesc : ImportOptions->Attributes)
@@ -606,7 +598,7 @@ UObject* USparseVolumeTextureFactory::ImportInternal(UClass* InClass, UObject* I
 	}
 
 	// Utility function for computing the bounding box encompassing the bounds of all frames in the SVT.
-	auto ExpandVolumeBounds = [](const FOpenVDBImportOptions& ImportOptions, const TArray<FOpenVDBGridInfo>& GridInfoArray, FIntVector3& VolumeBoundsMin, FIntVector3& VolumeBoundsMax, float* TransformMin, float* TransformMax)
+	auto ExpandVolumeBounds = [](const FOpenVDBImportOptions& ImportOptions, const TArray<FOpenVDBGridInfo>& GridInfoArray, FIntVector3& VolumeBoundsMin, FIntVector3& VolumeBoundsMax)
 	{
 		for (const FOpenVDBSparseVolumeAttributesDesc& Attributes : ImportOptions.Attributes)
 		{
@@ -622,15 +614,6 @@ UObject* USparseVolumeTextureFactory::ImportInternal(UClass* InClass, UObject* I
 					VolumeBoundsMax.X = FMath::Max(VolumeBoundsMax.X, GridInfo.VolumeActiveAABBMax.X);
 					VolumeBoundsMax.Y = FMath::Max(VolumeBoundsMax.Y, GridInfo.VolumeActiveAABBMax.Y);
 					VolumeBoundsMax.Z = FMath::Max(VolumeBoundsMax.Z, GridInfo.VolumeActiveAABBMax.Z);
-
-					for (int32 y = 0; y < 4; ++y)
-					{
-						for (int32 x = 0; x < 4; ++x)
-						{
-							TransformMin[y * 4 + x] = FMath::Min(TransformMin[y * 4 + x], GridInfo.Transform.M[y][x]);
-							TransformMax[y * 4 + x] = FMath::Max(TransformMax[y * 4 + x], GridInfo.Transform.M[y][x]);
-						}
-					}
 				}
 			}
 		}
@@ -638,13 +621,6 @@ UObject* USparseVolumeTextureFactory::ImportInternal(UClass* InClass, UObject* I
 
 	FIntVector3 VolumeBoundsMin = FIntVector3(INT32_MAX, INT32_MAX, INT32_MAX);
 	FIntVector3 VolumeBoundsMax = FIntVector3(INT32_MIN, INT32_MIN, INT32_MIN);
-	float TransformMin[16];
-	float TransformMax[16];
-	for (int32 i = 0; i < 16; ++i)
-	{
-		TransformMin[i] = FLT_MAX;
-		TransformMax[i] = -FLT_MAX;
-	}
 
 	// Import as either single static SVT or a sequence of frames, making up an animated SVT
 	if (!ImportOptions.bIsSequence)
@@ -654,10 +630,11 @@ UObject* USparseVolumeTextureFactory::ImportInternal(UClass* InClass, UObject* I
 		FScopedSlowTask ImportTask(1.0f, LOCTEXT("ImportingVDBStatic", "Importing static OpenVDB"));
 		ImportTask.MakeDialog(true);
 
-		ExpandVolumeBounds(ImportOptions, PreviewData.GridInfo, VolumeBoundsMin, VolumeBoundsMax, TransformMin, TransformMax);
+		ExpandVolumeBounds(ImportOptions, PreviewData.GridInfo, VolumeBoundsMin, VolumeBoundsMax);
 
 		UE::SVT::FTextureData TextureData{};
-		const bool bConversionSuccess = ConvertOpenVDBToSparseVolumeTexture(PreviewData.LoadedFile, ImportOptions, VolumeBoundsMin, false /*bBakeTranslation*/, TextureData);
+		FTransform FrameTransform = FTransform::Identity;
+		const bool bConversionSuccess = ConvertOpenVDBToSparseVolumeTexture(PreviewData.LoadedFile, ImportOptions, VolumeBoundsMin, TextureData, FrameTransform);
 
 		if (!bConversionSuccess)
 		{
@@ -666,7 +643,7 @@ UObject* USparseVolumeTextureFactory::ImportInternal(UClass* InClass, UObject* I
 		}
 
 		UStaticSparseVolumeTexture* StaticSVTexture = NewObject<UStaticSparseVolumeTexture>(InParent, UStaticSparseVolumeTexture::StaticClass(), InName, Flags);
-		const bool bInitSuccess = StaticSVTexture->Initialize(MakeArrayView<UE::SVT::FTextureData>(&TextureData, 1));
+		const bool bInitSuccess = StaticSVTexture->Initialize(MakeArrayView(&TextureData, 1), MakeArrayView(&FrameTransform, 1));
 		if (!bInitSuccess)
 		{
 			UE_LOG(LogSparseVolumeTextureFactory, Error, TEXT("Failed to initialize SparseVolumeTexture: %s"), *Filename);
@@ -696,14 +673,16 @@ UObject* USparseVolumeTextureFactory::ImportInternal(UClass* InClass, UObject* I
 
 		// Allocate space for each frame
 		TArray<UE::SVT::FTextureData> UncookedFramesData;
+		TArray<FTransform> FrameTransforms;
 		UncookedFramesData.SetNum(NumFrames);
+		FrameTransforms.SetNum(NumFrames);
 
 		std::atomic_bool bErrored = false;
 		std::atomic_bool bCanceled = false;
 
 		// Compute volume bounds and check sequence files for compatiblity
 		std::mutex VolumeBoundsMutex;
-		ParallelFor(NumFrames, [&bErrored, &VolumeBoundsMutex, &VolumeBoundsMin, &VolumeBoundsMax, TransformMin = &TransformMin[0], TransformMax = &TransformMax[0], &ExpandVolumeBounds, &PreviewData, &ImportOptions](int32 FrameIdx)
+		ParallelFor(NumFrames, [&bErrored, &VolumeBoundsMutex, &VolumeBoundsMin, &VolumeBoundsMax, &ExpandVolumeBounds, &PreviewData, &ImportOptions](int32 FrameIdx)
 			{
 				if (bErrored.load())
 				{
@@ -757,7 +736,7 @@ UObject* USparseVolumeTextureFactory::ImportInternal(UClass* InClass, UObject* I
 				// Update sequence volume bounds and increment ProcessedFramesCounter
 				{
 					std::lock_guard<std::mutex> Lock(VolumeBoundsMutex);
-					ExpandVolumeBounds(ImportOptions, FrameGridInfo, VolumeBoundsMin, VolumeBoundsMax, TransformMin, TransformMax);
+					ExpandVolumeBounds(ImportOptions, FrameGridInfo, VolumeBoundsMin, VolumeBoundsMax);
 				}
 			});
 
@@ -766,39 +745,6 @@ UObject* USparseVolumeTextureFactory::ImportInternal(UClass* InClass, UObject* I
 			return nullptr;
 		}
 
-		FMatrix44f TransformMatrixMin;
-		FMatrix44f TransformMatrixMax;
-		for (int32 y = 0; y < 4; ++y)
-		{
-			for (int32 x = 0; x < 4; ++x)
-			{
-				TransformMatrixMin.M[y][x] = TransformMin[y * 4 + x];
-				TransformMatrixMax.M[y][x] = TransformMax[y * 4 + x];
-			}
-		}
-		const FVector3f TransformScaleMin = TransformMatrixMin.GetScaleVector();
-		const FVector3f TransformScaleMax = TransformMatrixMax.GetScaleVector();
-		const FVector3f TransformTransMin = TransformMatrixMin.GetOrigin();
-		const FVector3f TransformTransMax = TransformMatrixMax.GetOrigin();
-		const bool bSameScaleForAllFrames = (TransformScaleMax - TransformScaleMin).GetAbsMax() <= UE_SMALL_NUMBER;
-		const bool bSameTranslationForAllFrames = (TransformTransMax - TransformTransMin).GetAbsMax() <= UE_SMALL_NUMBER;
-
-		// We currently do not support changing scale/rotation across frames
-		if (!bSameScaleForAllFrames)
-		{
-			UE_LOG(LogSparseVolumeTextureFactory, Warning, TEXT("OpenVDB sequence is using different scalings and rotations across all frames"));
-		}
-
-		const bool bBakeTranslations = !bSameTranslationForAllFrames && (GSVTImportBakeFrameTranslations != 0);
-
-		// But we can bake translation changes into the SVT
-		if (bBakeTranslations)
-		{
-			FVector3f OffsetF = TransformTransMin / TransformScaleMin;
-			FIntVector3 Offset = FIntVector3(FMath::Floor(OffsetF.X), FMath::Floor(OffsetF.Y), FMath::Floor(OffsetF.Z));
-			VolumeBoundsMin += Offset;
-		}
-
 		ImportTask.EnterProgressFrame(1.0f, LOCTEXT("ConvertingVDBAnim", "Converting OpenVDB animation"));
 
 		FEvent* AllTasksFinishedEvent = FPlatformProcess::GetSynchEventFromPool();
@@ -826,8 +772,8 @@ UObject* USparseVolumeTextureFactory::ImportInternal(UClass* InClass, UObject* I
 			};
 
 			AsyncTask(ENamedThreads::AnyNormalThreadNormalTask,
-				[FrameIdx, NumFrames, &PreviewData, &ImportOptions, &UncookedFramesData,
-				AllTasksFinishedEvent, &bErrored, &bCanceled, &FinishedTasksCounter, &ProcessedFramesCounter, &VolumeBoundsMin, bBakeTranslations]()
+				[FrameIdx, NumFrames, &PreviewData, &ImportOptions, &UncookedFramesData, &FrameTransforms,
+				AllTasksFinishedEvent, &bErrored, &bCanceled, &FinishedTasksCounter, &ProcessedFramesCounter, &VolumeBoundsMin]()
 				{
 					// Ensure the FinishedTasksCounter will be incremented in all cases
 					FScopedIncrementer Incremeter(FinishedTasksCounter, NumFrames, AllTasksFinishedEvent);
@@ -851,7 +797,8 @@ UObject* USparseVolumeTextureFactory::ImportInternal(UClass* InClass, UObject* I
 					}
 
 					UE::SVT::FTextureData TextureData{};
-					const bool bConversionSuccess = ConvertOpenVDBToSparseVolumeTexture(LoadedFrameFile, ImportOptions, VolumeBoundsMin, bBakeTranslations, TextureData);
+					FTransform FrameTransform = FTransform::Identity;
+					const bool bConversionSuccess = ConvertOpenVDBToSparseVolumeTexture(LoadedFrameFile, ImportOptions, VolumeBoundsMin, TextureData, FrameTransform);
 
 					if (!bConversionSuccess)
 					{
@@ -861,6 +808,7 @@ UObject* USparseVolumeTextureFactory::ImportInternal(UClass* InClass, UObject* I
 					}
 
 					UncookedFramesData[FrameIdx] = MoveTemp(TextureData);
+					FrameTransforms[FrameIdx] = FrameTransform;
 
 					// Increment ProcessedFramesCounter
 					ProcessedFramesCounter.fetch_add(1);
@@ -942,7 +890,7 @@ UObject* USparseVolumeTextureFactory::ImportInternal(UClass* InClass, UObject* I
 		}
 
 		UAnimatedSparseVolumeTexture* AnimatedSVTexture = NewObject<UAnimatedSparseVolumeTexture>(InParent, UAnimatedSparseVolumeTexture::StaticClass(), NewObjectName, Flags);
-		const bool bInitSuccess = AnimatedSVTexture->Initialize(UncookedFramesData);
+		const bool bInitSuccess = AnimatedSVTexture->Initialize(UncookedFramesData, FrameTransforms);
 		if (!bInitSuccess)
 		{
 			UE_LOG(LogSparseVolumeTextureFactory, Error, TEXT("Failed to initialize SparseVolumeTexture: %s"), *Filename);
diff --git a/Engine/Source/Editor/SparseVolumeTexture/Private/SparseVolumeTextureOpenVDBUtility.cpp b/Engine/Source/Editor/SparseVolumeTexture/Private/SparseVolumeTextureOpenVDBUtility.cpp
--- a/Engine/Source/Editor/SparseVolumeTexture/Private/SparseVolumeTextureOpenVDBUtility.cpp
+++ b/Engine/Source/Editor/SparseVolumeTexture/Private/SparseVolumeTextureOpenVDBUtility.cpp
@@ -91,7 +91,6 @@ static FOpenVDBGridInfo GetOpenVDBGridInfo(openvdb::GridBase::Ptr Grid, uint32 G
 {
 	openvdb::CoordBBox VolumeActiveAABB = Grid->evalActiveVoxelBoundingBox();
 	openvdb::Coord VolumeActiveDim = Grid->evalActiveVoxelDim();
-	openvdb::Vec3d VolumeVoxelSize = Grid->voxelSize();
 	openvdb::math::MapBase::ConstPtr MapBase = Grid->constTransform().baseMap();
 	openvdb::Vec3d VoxelSize = MapBase->voxelSize();
 	openvdb::Mat4d GridTransformVDB = MapBase->getAffineMap()->getConstMat4();
@@ -103,17 +102,17 @@ static FOpenVDBGridInfo GetOpenVDBGridInfo(openvdb::GridBase::Ptr Grid, uint32 G
 	GridInfo.VolumeActiveAABBMin = FIntVector3(VolumeActiveAABB.min().x(), VolumeActiveAABB.min().y(), VolumeActiveAABB.min().z());
 	GridInfo.VolumeActiveAABBMax = FIntVector3(VolumeActiveAABB.max().x() + 1, VolumeActiveAABB.max().y() + 1, VolumeActiveAABB.max().z() + 1); // +1 because CoordBBox::Max is inclusive, but we want exclusive
 	GridInfo.VolumeActiveDim = FIntVector3(VolumeActiveDim.x(), VolumeActiveDim.y(), VolumeActiveDim.z());
-	GridInfo.VolumeVoxelSize = FVector(VolumeVoxelSize.x(), VolumeVoxelSize.y(), VolumeVoxelSize.z());
 	GridInfo.bIsInWorldSpace = Grid->isInWorldSpace();
-	GridInfo.bHasUniformVoxels = Grid->hasUniformVoxels();
 
+	FMatrix TransformMatrix;
 	for (int i = 0; i < 4; ++i)
 	{
 		for (int j = 0; j < 4; ++j)
 		{
-			GridInfo.Transform.M[i][j] = static_cast<float>(GridTransformVDB[i][j]);
+			TransformMatrix.M[i][j] = GridTransformVDB[i][j];
 		}
 	}
+	GridInfo.Transform = FTransform(TransformMatrix);
 
 	// Figure out the type/format of the grid
 	if (Grid->isType<FOpenVDBHalf1Grid>())
@@ -196,11 +195,6 @@ static FOpenVDBGridInfo GetOpenVDBGridInfo(openvdb::GridBase::Ptr Grid, uint32 G
 
 bool IsOpenVDBGridValid(const FOpenVDBGridInfo& GridInfo, const FString& Filename)
 {
-	if (!GridInfo.bHasUniformVoxels)
-	{
-		UE_LOG(LogSparseVolumeTextureOpenVDBUtility, Warning, TEXT("OpenVDB importer cannot handle non uniform voxels: %s"), *Filename);
-		return false;
-	}
 	return true;
 }
 
@@ -285,11 +279,10 @@ class FSparseVolumeTextureDataProviderOpenVDB : public UE::SVT::ITextureDataProv
 {
 public:
 
-	bool Initialize(TArray64<uint8>& SourceFile, const FOpenVDBImportOptions& ImportOptions, const FIntVector3& InVolumeBoundsMin, bool bInBakeTranslation)
+	bool Initialize(TArray64<uint8>& SourceFile, const FOpenVDBImportOptions& ImportOptions, const FIntVector3& InVolumeBoundsMin)
 	{
 		Attributes = ImportOptions.Attributes;
 		VolumeBoundsMin = InVolumeBoundsMin;
-		Translation = FIntVector3::ZeroValue;
 
 		// Compute some basic info about the number of components and which format to use
 		EPixelFormat MultiCompFormat[NumAttributesDescs] = {};
@@ -367,10 +360,7 @@ public:
 
 		FIntVector3 SmallestAABBMin = FIntVector3(INT32_MAX);
 		FIntVector3 LargestAABBMax = FIntVector3(INT32_MIN);
-		FVector3f ScaleMin = FVector3f(FLT_MAX);
-		FVector3f ScaleMax = FVector3f(-FLT_MAX);
-		FVector3f TranslationMin = FVector3f(FLT_MAX);
-		FVector3f TranslationMax = FVector3f(-FLT_MAX);
+		const FTransform* LastGridTransformPtr = nullptr;
 
 		// Compute per source grid data of up to 4 different grids (one per component)
 		UniqueGridAdapters.SetNum((int32)NumSourceGrids);
@@ -412,12 +402,11 @@ public:
 				LargestAABBMax.Y = FMath::Max(LargestAABBMax.Y, GridInfo.VolumeActiveAABBMax.Y);
 				LargestAABBMax.Z = FMath::Max(LargestAABBMax.Z, GridInfo.VolumeActiveAABBMax.Z);
 
-				const FVector3f GridScale = GridInfo.Transform.GetScaleVector();
-				const FVector3f GridTranslation = GridInfo.Transform.GetOrigin();
-				ScaleMin = FVector3f::Min(ScaleMin, GridScale);
-				ScaleMax = FVector3f::Max(ScaleMax, GridScale);
-				TranslationMin = FVector3f::Min(TranslationMin, GridTranslation);
-				TranslationMax = FVector3f::Max(TranslationMax, GridTranslation);
+				if (LastGridTransformPtr && !LastGridTransformPtr->Equals(GridInfo.Transform))
+				{
+					UE_LOG(LogSparseVolumeTextureOpenVDBUtility, Warning, TEXT("Frame has multiple grids with different transforms in the same frame! Data will likely not be imported/displayed correctly!"));
+				}
+				LastGridTransformPtr = &GridInfo.Transform;
 
 				SVTCreateInfo.FallbackValues[AttributesIdx][CompIdx] = UniqueGridAdapters[SourceGridIndex]->GetBackgroundValue(SourceComponentIndex);
 
@@ -429,25 +418,12 @@ public:
 			}
 		}
 
-		const float ScaleDiffAbs = FVector3f(ScaleMax - ScaleMin).GetAbsMax();
-		const float TransDiffAbs = FVector3f(TranslationMax - TranslationMin).GetAbsMax();
-		if (bInBakeTranslation && ((ScaleDiffAbs > UE_SMALL_NUMBER) || (TransDiffAbs > UE_SMALL_NUMBER)))
-		{
-			UE_LOG(LogSparseVolumeTextureOpenVDBUtility, Warning, TEXT("Frame has multiple grids with different transforms!"));
-		}
-
 		const bool bEmptyBounds = SmallestAABBMin.X >= LargestAABBMax.X || SmallestAABBMin.Y >= LargestAABBMax.Y || SmallestAABBMin.Z >= LargestAABBMax.Z;
-
-		if (bInBakeTranslation && !bEmptyBounds)
-		{
-			FVector3f OffsetF = TranslationMin / ScaleMin;
-			Translation = FIntVector3(FMath::Floor(OffsetF.X), FMath::Floor(OffsetF.Y), FMath::Floor(OffsetF.Z));
-			SmallestAABBMin += Translation;
-			LargestAABBMax += Translation;
-		}
-
 		SVTCreateInfo.VirtualVolumeAABBMin = bEmptyBounds ? FIntVector3(INT32_MAX) : (SmallestAABBMin - VolumeBoundsMin);
 		SVTCreateInfo.VirtualVolumeAABBMax = bEmptyBounds ? FIntVector3(INT32_MIN) : (LargestAABBMax - VolumeBoundsMin);
+		
+		FrameTransform = LastGridTransformPtr ? *LastGridTransformPtr : FTransform::Identity;
+		FrameTransform.AddToTranslation(FVector(VolumeBoundsMin) * FrameTransform.GetScale3D()); // We made the volume relative to VolumeBoundsMin, so we need to undo this translation when using the frame transform
 
 		return true;
 	}
@@ -469,7 +445,7 @@ public:
 			UniqueGridAdapters[GridIdx]->IteratePhysical(
 				[&](const FIntVector3& Coord, uint32 NumVoxelComponents, float* VoxelValues)
 				{
-					FIntVector3 RemappedCoord = Coord + Translation - VolumeBoundsMin;
+					FIntVector3 RemappedCoord = Coord - VolumeBoundsMin;
 					for (const FSingleGridToComponentMapping& Mapping : GridToComponentMappings[GridIdx])
 					{
 						OnVisit(RemappedCoord, Mapping.AttributesIdx, Mapping.ComponentIdx, VoxelValues[Mapping.GridComponentIdx]);
@@ -478,6 +454,8 @@ public:
 		}
 	}
 
+	FTransform GetFrameTransform() const { return FrameTransform; }
+
 private:
 
 	struct FSingleGridToComponentMapping
@@ -494,16 +472,16 @@ private:
 	TStaticArray<uint32, NumAttributesDescs> NumComponents;
 	UE::SVT::FTextureDataCreateInfo SVTCreateInfo;
 	FIntVector3 VolumeBoundsMin;
-	FIntVector3 Translation;
+	FTransform FrameTransform;
 };
 
 #endif // OPENVDB_AVAILABLE
 
-bool ConvertOpenVDBToSparseVolumeTexture(TArray64<uint8>& SourceFile, const FOpenVDBImportOptions& ImportOptions, const FIntVector3& VolumeBoundsMin, bool bBakeTranslation, UE::SVT::FTextureData& OutResult)
+bool ConvertOpenVDBToSparseVolumeTexture(TArray64<uint8>& SourceFile, const FOpenVDBImportOptions& ImportOptions, const FIntVector3& VolumeBoundsMin, UE::SVT::FTextureData& OutResult, FTransform& OutFrameTransform)
 {
 #if OPENVDB_AVAILABLE
 	FSparseVolumeTextureDataProviderOpenVDB DataProvider;
-	if (!DataProvider.Initialize(SourceFile, ImportOptions, VolumeBoundsMin, bBakeTranslation))
+	if (!DataProvider.Initialize(SourceFile, ImportOptions, VolumeBoundsMin))
 	{
 		return false;
 	}
@@ -511,6 +489,7 @@ bool ConvertOpenVDBToSparseVolumeTexture(TArray64<uint8>& SourceFile, const FOpe
 	{
 		return false;
 	}
+	OutFrameTransform = DataProvider.GetFrameTransform();
 	return true;
 #else
 	return false;
diff --git a/Engine/Source/Editor/SparseVolumeTexture/Private/SparseVolumeTextureOpenVDBUtility.h b/Engine/Source/Editor/SparseVolumeTexture/Private/SparseVolumeTextureOpenVDBUtility.h
--- a/Engine/Source/Editor/SparseVolumeTexture/Private/SparseVolumeTextureOpenVDBUtility.h
+++ b/Engine/Source/Editor/SparseVolumeTexture/Private/SparseVolumeTextureOpenVDBUtility.h
@@ -28,7 +28,7 @@ bool IsOpenVDBGridValid(const FOpenVDBGridInfo& GridInfo, const FString& Filenam
 
 bool GetOpenVDBGridInfo(TArray64<uint8>& SourceFile, bool bCreateStrings, TArray<FOpenVDBGridInfo>* OutGridInfo);
 
-bool ConvertOpenVDBToSparseVolumeTexture(TArray64<uint8>& SourceFile, const struct FOpenVDBImportOptions& ImportOptions, const FIntVector3& VolumeBoundsMin, bool bBakeTranslation, UE::SVT::FTextureData& OutResult);
+bool ConvertOpenVDBToSparseVolumeTexture(TArray64<uint8>& SourceFile, const struct FOpenVDBImportOptions& ImportOptions, const FIntVector3& VolumeBoundsMin, UE::SVT::FTextureData& OutResult, FTransform& OutFrameTransform);
 
 const TCHAR* OpenVDBGridTypeToString(EOpenVDBGridType Type);
 
diff --git a/Engine/Source/Editor/SparseVolumeTexture/Public/OpenVDBImportOptions.h b/Engine/Source/Editor/SparseVolumeTexture/Public/OpenVDBImportOptions.h
--- a/Engine/Source/Editor/SparseVolumeTexture/Public/OpenVDBImportOptions.h
+++ b/Engine/Source/Editor/SparseVolumeTexture/Public/OpenVDBImportOptions.h
@@ -51,18 +51,16 @@ enum class EOpenVDBGridType : uint8
 
 struct FOpenVDBGridInfo
 {
-	FMatrix44f Transform;
+	FTransform Transform;
 	FIntVector3 VolumeActiveAABBMin;
 	FIntVector3 VolumeActiveAABBMax;
 	FIntVector3 VolumeActiveDim;
-	FVector VolumeVoxelSize;
 	FString Name;
 	FString DisplayString; // Contains Index (into source file grids), Type and Name
 	uint32 Index;
 	uint32 NumComponents;
 	EOpenVDBGridType Type;
 	bool bIsInWorldSpace;
-	bool bHasUniformVoxels;
 };
 
 struct FOpenVDBGridComponentInfo
diff --git a/Engine/Source/Runtime/Engine/Classes/SparseVolumeTexture/SparseVolumeTexture.h b/Engine/Source/Runtime/Engine/Classes/SparseVolumeTexture/SparseVolumeTexture.h
--- a/Engine/Source/Runtime/Engine/Classes/SparseVolumeTexture/SparseVolumeTexture.h
+++ b/Engine/Source/Runtime/Engine/Classes/SparseVolumeTexture/SparseVolumeTexture.h
@@ -259,6 +259,9 @@ public:
 	UFUNCTION(BlueprintCallable, Category = "Texture")
 	virtual int32 GetNumMipLevels() const { return 0; }
 
+	UFUNCTION(BlueprintCallable, Category = "Texture")
+	virtual FTransform GetFrameTransform() const { return FTransform::Identity; }
+
 	virtual FIntVector GetVolumeResolution() const { return FIntVector(); }
 	virtual EPixelFormat GetFormat(int32 AttributesIndex) const { return PF_Unknown; }
 	virtual FVector4f GetFallbackValue(int32 AttributesIndex) const { return FVector4f(); }
@@ -306,7 +309,7 @@ public:
 	// If streaming cooked data from disk, the highest priority will be used, but no guarantee is given.
 	static ENGINE_API USparseVolumeTextureFrame* GetFrameAndIssueStreamingRequest(USparseVolumeTexture* SparseVolumeTexture, float FrameIndex, int32 MipLevel, bool bBlocking);
 
-	ENGINE_API bool Initialize(USparseVolumeTexture* InOwner, int32 InFrameIndex, UE::SVT::FTextureData& UncookedFrame);
+	ENGINE_API bool Initialize(USparseVolumeTexture* InOwner, int32 InFrameIndex, const FTransform& InFrameTransform, UE::SVT::FTextureData& UncookedFrame);
 	int32 GetFrameIndex() const { return FrameIndex; }
 	UE::SVT::FResources* GetResources() { return &Resources; }
 	// Creates TextureRenderResources if they don't already exist. Returns false if they already existed.
@@ -335,6 +338,7 @@ public:
 	//~ Begin USparseVolumeTexture Interface.
 	virtual int32 GetNumFrames() const override { return 1; }
 	virtual int32 GetNumMipLevels() const override { return Owner->GetNumMipLevels(); }
+	virtual FTransform GetFrameTransform() const override { return Transform; }
 	virtual FIntVector GetVolumeResolution() const override { return Owner->GetVolumeResolution(); }
 	virtual EPixelFormat GetFormat(int32 AttributesIndex) const override { return Owner->GetFormat(AttributesIndex); }
 	virtual FVector4f GetFallbackValue(int32 AttributesIndex) const override { return Owner->GetFallbackValue(AttributesIndex); }
@@ -352,6 +356,9 @@ private:
 	UPROPERTY()
 	int32 FrameIndex;
 
+	UPROPERTY(VisibleAnywhere, Category = "Texture", AssetRegistrySearchable)
+	FTransform Transform;
+
 #if WITH_EDITORONLY_DATA
 	// FTextureData from which the FResources data can be built with a call to FResources::Build()
 	UE::Serialization::FEditorBulkData SourceData;
@@ -424,11 +431,11 @@ public:
 	// The NumExpectedFrames parameter on BeginInitialize() just serves as a potential optimization to reserve memory for the frames to be appended
 	// and doesn't need to match the exact number if it is not known at the time.
 	ENGINE_API virtual bool BeginInitialize(int32 NumExpectedFrames);
-	ENGINE_API virtual bool AppendFrame(UE::SVT::FTextureData& UncookedFrame);
+	ENGINE_API virtual bool AppendFrame(UE::SVT::FTextureData& UncookedFrame, const FTransform& FrameTransform);
 	ENGINE_API virtual bool EndInitialize();
 
 	// Convenience function wrapping the multi-phase initialization functions above
-	virtual bool Initialize(const TArrayView<UE::SVT::FTextureData>& UncookedData);
+	virtual bool Initialize(const TArrayView<UE::SVT::FTextureData>& UncookedData, const TArrayView<FTransform>& FrameTransforms);
 	// Consider using USparseVolumeTextureFrame::GetFrameAndIssueStreamingRequest() if the frame should have streaming requests issued.
 	USparseVolumeTextureFrame* GetFrame(int32 FrameIndex) const { return Frames.IsValidIndex(FrameIndex) ? Frames[FrameIndex] : nullptr; }
 
@@ -457,6 +464,7 @@ public:
 	//~ Begin USparseVolumeTexture Interface.
 	virtual int32 GetNumFrames() const override { return Frames.Num(); }
 	virtual int32 GetNumMipLevels() const override { return NumMipLevels; }
+	virtual FTransform GetFrameTransform() const override { return Frames.IsEmpty() ? FTransform::Identity : Frames[0]->GetFrameTransform(); }
 	virtual FIntVector GetVolumeResolution() const override { return VolumeResolution; };
 	virtual EPixelFormat GetFormat(int32 AttributesIndex) const override { check(AttributesIndex >= 0 && AttributesIndex < 2) return AttributesIndex == 0 ? FormatA : FormatB; }
 	virtual FVector4f GetFallbackValue(int32 AttributesIndex) const override { check(AttributesIndex >= 0 && AttributesIndex < 2) return AttributesIndex == 0 ? FallbackValueA : FallbackValueB; }
@@ -509,7 +517,7 @@ public:
 
 	//~ Begin UStreamableSparseVolumeTexture Interface.
 	// Override AppendFrame() to ensure that there is never more than a single frame in a static SVT
-	ENGINE_API virtual bool AppendFrame(UE::SVT::FTextureData& UncookedFrame) override;
+	ENGINE_API virtual bool AppendFrame(UE::SVT::FTextureData& UncookedFrame, const FTransform& InFrameTransform) override;
 	//~ End UStreamableSparseVolumeTexture Interface.
 
 	//~ Begin USparseVolumeTexture Interface.
diff --git a/Engine/Source/Runtime/Engine/Private/SparseVolumeTexture/SparseVolumeTexture.cpp b/Engine/Source/Runtime/Engine/Private/SparseVolumeTexture/SparseVolumeTexture.cpp
--- a/Engine/Source/Runtime/Engine/Private/SparseVolumeTexture/SparseVolumeTexture.cpp
+++ b/Engine/Source/Runtime/Engine/Private/SparseVolumeTexture/SparseVolumeTexture.cpp
@@ -1062,11 +1062,12 @@ USparseVolumeTextureFrame* USparseVolumeTextureFrame::GetFrameAndIssueStreamingR
 	return nullptr;
 }
 
-bool USparseVolumeTextureFrame::Initialize(USparseVolumeTexture* InOwner, int32 InFrameIndex, UE::SVT::FTextureData& UncookedFrame)
+bool USparseVolumeTextureFrame::Initialize(USparseVolumeTexture* InOwner, int32 InFrameIndex, const FTransform& InFrameTransform, UE::SVT::FTextureData& UncookedFrame)
 {
 #if WITH_EDITORONLY_DATA
 	Owner = InOwner;
 	FrameIndex = InFrameIndex;
+	Transform = InFrameTransform;
 	{
 		UE::Serialization::FEditorBulkDataWriter SourceDataArchiveWriter(SourceData);
 		SourceDataArchiveWriter << UncookedFrame;
@@ -1255,7 +1256,7 @@ bool UStreamableSparseVolumeTexture::BeginInitialize(int32 NumExpectedFrames)
 #endif
 }
 
-bool UStreamableSparseVolumeTexture::AppendFrame(UE::SVT::FTextureData& UncookedFrame)
+bool UStreamableSparseVolumeTexture::AppendFrame(UE::SVT::FTextureData& UncookedFrame, const FTransform& FrameTransform)
 {
 #if WITH_EDITORONLY_DATA
 	if (InitState != EInitState_Pending)
@@ -1273,9 +1274,14 @@ bool UStreamableSparseVolumeTexture::AppendFrame(UE::SVT::FTextureData& Uncooked
 		return false;
 	}
 
-	// SVT_TODO: Valide formats against list of supported formats
 	if (Frames.IsEmpty())
 	{
+		if (!UE::SVT::IsSupportedFormat(UncookedFrame.Header.AttributesFormats[0]) || !UE::SVT::IsSupportedFormat(UncookedFrame.Header.AttributesFormats[1]))
+		{
+			UE_LOG(LogSparseVolumeTexture, Error, TEXT("Tried to add a frame with unsupported formats to a SparseVolumeTexture! Formats: (%i %i)"),
+				(int)UncookedFrame.Header.AttributesFormats[0], (int)UncookedFrame.Header.AttributesFormats[1]);
+			return false;
+		}
 		FormatA = UncookedFrame.Header.AttributesFormats[0];
 		FormatB = UncookedFrame.Header.AttributesFormats[1];
 		FallbackValueA = UncookedFrame.Header.FallbackValues[0];
@@ -1316,7 +1322,7 @@ bool UStreamableSparseVolumeTexture::AppendFrame(UE::SVT::FTextureData& Uncooked
 	
 
 	USparseVolumeTextureFrame* Frame = NewObject<USparseVolumeTextureFrame>(this);
-	if (Frame->Initialize(this, Frames.Num(), UncookedFrame))
+	if (Frame->Initialize(this, Frames.Num(), FrameTransform, UncookedFrame))
 	{
 		Frames.Add(Frame);
 		return true;
@@ -1343,7 +1349,7 @@ bool UStreamableSparseVolumeTexture::EndInitialize()
 		UE_LOG(LogSparseVolumeTexture, Warning, TEXT("SVT has zero frames! Adding a dummy frame. SVT: %s"), *GetName());
 		UE::SVT::FTextureData DummyFrame;
 		DummyFrame.CreateDefault();
-		AppendFrame(DummyFrame);
+		AppendFrame(DummyFrame, FTransform::Identity);
 	}
 
 	check(VolumeResolution.X > 0 && VolumeResolution.Y > 0 && VolumeResolution.Z > 0);
@@ -1375,7 +1381,7 @@ bool UStreamableSparseVolumeTexture::EndInitialize()
 #endif
 }
 
-bool UStreamableSparseVolumeTexture::Initialize(const TArrayView<UE::SVT::FTextureData>& InUncookedData)
+bool UStreamableSparseVolumeTexture::Initialize(const TArrayView<UE::SVT::FTextureData>& InUncookedData, const TArrayView<FTransform>& InFrameTransforms)
 {
 	if (InUncookedData.IsEmpty())
 	{
@@ -1383,13 +1389,16 @@ bool UStreamableSparseVolumeTexture::Initialize(const TArrayView<UE::SVT::FTextu
 		return false;
 	}
 
-	if (!BeginInitialize(InUncookedData.Num()))
+	const int32 NumUncookedFrameData = InUncookedData.Num();
+	const int32 NumFrameTransforms = InFrameTransforms.Num();
+	const bool bHasValidFrameTransforms = NumUncookedFrameData <= NumFrameTransforms;
+	if (!BeginInitialize(NumUncookedFrameData))
 	{
 		return false;
 	}
-	for (UE::SVT::FTextureData& UncookedFrame : InUncookedData)
+	for (int32 i = 0; i < NumUncookedFrameData; ++i)
 	{
-		if (!AppendFrame(UncookedFrame))
+		if (!AppendFrame(InUncookedData[i], bHasValidFrameTransforms ? InFrameTransforms[i] : FTransform::Identity))
 		{
 			return false;
 		}
@@ -1600,14 +1609,14 @@ UStaticSparseVolumeTexture::UStaticSparseVolumeTexture(const FObjectInitializer&
 {
 }
 
-bool UStaticSparseVolumeTexture::AppendFrame(UE::SVT::FTextureData& UncookedFrame)
+bool UStaticSparseVolumeTexture::AppendFrame(UE::SVT::FTextureData& UncookedFrame, const FTransform& InFrameTransform)
 {
 	if (!Frames.IsEmpty())
 	{
 		UE_LOG(LogSparseVolumeTexture, Error, TEXT("Tried to initialize a UStaticSparseVolumeTexture with more than 1 frame"));
 		return false;
 	}
-	return Super::AppendFrame(UncookedFrame);
+	return Super::AppendFrame(UncookedFrame, InFrameTransform);
 }
 
 ////////////////////////////////////////////////////////////////////////////////////////////////
diff --git a/Engine/Source/Runtime/Renderer/Private/SparseVolumeTexture/SparseVolumeTextureViewerComponent.cpp b/Engine/Source/Runtime/Renderer/Private/SparseVolumeTexture/SparseVolumeTextureViewerComponent.cpp
--- a/Engine/Source/Runtime/Renderer/Private/SparseVolumeTexture/SparseVolumeTextureViewerComponent.cpp
+++ b/Engine/Source/Runtime/Renderer/Private/SparseVolumeTexture/SparseVolumeTextureViewerComponent.cpp
@@ -20,7 +20,7 @@
 
 #define LOCTEXT_NAMESPACE "SparseVolumeTextureViewerComponent"
 
-constexpr double SVTViewerDefaultVolumeExtent = 100.0;
+constexpr int32 SVTViewerDefaultVolumeResolution = 128;
 
 /*=============================================================================
 	USparseVolumeTextureViewerComponent implementation.
@@ -71,26 +71,26 @@ void USparseVolumeTextureViewerComponent::Serialize(FArchive& Ar)
 
 FBoxSphereBounds USparseVolumeTextureViewerComponent::CalcBounds(const FTransform& LocalToWorld) const
 {
-	FBoxSphereBounds NormalizedBound;
-	NormalizedBound.Origin = FVector(0.0f, 0.0f, 0.0f);
-
+	FVector VolumeResolution = FVector(SVTViewerDefaultVolumeResolution);
 	if (SparseVolumeTexturePreview)
 	{
-		// We assume that the maximum bound will have a length of SVTViewerDefaultVolumeExtent (1 meter). 
-		// Then the other dimensions are scaled relatively.
-		// All this account for the size of volume with page table (they add padding).
-		// TODO can we recover world size from OpenVDB meta data in meter?
-		const FVector VolumeResolution = FVector(SparseVolumeTexturePreview->GetVolumeResolution());
-		const double MaxDim = FMath::Max(FMath::Max(VolumeResolution.X, VolumeResolution.Y), VolumeResolution.Z);
-		NormalizedBound.BoxExtent = VolumeResolution / MaxDim * SVTViewerDefaultVolumeExtent;
+		VolumeResolution = FVector(SparseVolumeTexturePreview->GetVolumeResolution());
+	}
+	
+	FBoxSphereBounds NewBounds;
+	if (bLocalOriginAtCorner)
+	{
+		NewBounds.Origin = -VolumeResolution * 0.5;
+		NewBounds.BoxExtent = VolumeResolution;
 	}
 	else
 	{
-		NormalizedBound.BoxExtent = FVector(SVTViewerDefaultVolumeExtent);
+		NewBounds.Origin = FVector::ZeroVector;
+		NewBounds.BoxExtent = VolumeResolution * 0.5;
 	}
+	NewBounds.SphereRadius = NewBounds.BoxExtent.Length();
 
-	NormalizedBound.SphereRadius = NormalizedBound.BoxExtent.Size();
-	return NormalizedBound.TransformBy(LocalToWorld);
+	return NewBounds.TransformBy(LocalToWorld);
 }
 
 void USparseVolumeTextureViewerComponent::CreateRenderState_Concurrent(FRegisterComponentContext* Context)
@@ -147,44 +147,33 @@ void USparseVolumeTextureViewerComponent::SendRenderTransformCommand()
 {
 	if (SparseVolumeTextureViewerSceneProxy)
 	{
-		FVector VolumeExtent = FVector(SVTViewerDefaultVolumeExtent);
-		FVector VolumeResolution = FVector(SVTViewerDefaultVolumeExtent * 2.0);
+		FVector VolumeResolution = FVector(SVTViewerDefaultVolumeResolution);
+		FTransform FrameTransform = FTransform::Identity;
 		if (SparseVolumeTextureFrame)
 		{
 			VolumeResolution = FVector(SparseVolumeTextureFrame->GetVolumeResolution());
-			const double MaxBoundsDim = FMath::Max(FMath::Max(VolumeResolution.X, VolumeResolution.Y), VolumeResolution.Z);
-			VolumeExtent = VolumeResolution / MaxBoundsDim * SVTViewerDefaultVolumeExtent;
+			if (bApplyPerFrameTransforms)
+			{
+				FrameTransform = SparseVolumeTextureFrame->GetFrameTransform();
+			}
 		}
 
-		const FTransform ToWorldTransform = GetComponentTransform();
-		const FVector Scale3D = ToWorldTransform.GetScale3D();
-		// Keep max scaling since the DDA algorithm has trouble with non-uniform scale.
-		// Note that using the other components in FScaleMatrix below is fine 
-		// because they cancel out with the actual volume resolution, producing uniformly scaled voxels.
-		const float MaxScaling = FMath::Max(Scale3D.X, FMath::Max(Scale3D.Y, Scale3D.Z));
-
-		const FRotationMatrix WorldToLocalRotation = FRotationMatrix(FRotator(ToWorldTransform.GetRotation().Inverse()));
-		const FMatrix44f ToLocalMatNoScale = FMatrix44f(WorldToLocalRotation);
-		const FMatrix44f ToLocalMat = FMatrix44f(
-			FTranslationMatrix(-ToWorldTransform.GetTranslation()) 
-			* WorldToLocalRotation 
-			* FScaleMatrix((VolumeExtent * MaxScaling).Reciprocal()));
-
+		const FTransform GlobalTransform = GetComponentTransform();
 		const FVector3f VolumeRes3f = FVector3f(VolumeResolution.X, VolumeResolution.Y, VolumeResolution.Z);
 
-
-		FSparseVolumeTextureViewerSceneProxy* SVTViewerSceneProxy = SparseVolumeTextureViewerSceneProxy;
+		FSparseVolumeTextureViewerSceneProxy* SVTSceneProxy = SparseVolumeTextureViewerSceneProxy;
 		ENQUEUE_RENDER_COMMAND(FUpdateSparseVolumeTextureViewerProxyTransformCommand)(
-			[SVTViewerSceneProxy, ToLocalMat, ToLocalMatNoScale, CompIdx = (uint32)PreviewAttribute, Ext = Extinction, Mip = MipLevel, VolumeRes3f]
-		(FRHICommandList& RHICmdList)
-			{
-				SVTViewerSceneProxy->WorldToLocal = ToLocalMat;
-				SVTViewerSceneProxy->WorldToLocalNoScale = ToLocalMatNoScale;
-				SVTViewerSceneProxy->VolumeResolution = VolumeRes3f;
-				SVTViewerSceneProxy->MipLevel = Mip;
-				SVTViewerSceneProxy->ComponentToVisualize = CompIdx;
-				SVTViewerSceneProxy->Extinction = Ext;
-			});
+		[SVTSceneProxy, GlobalTransform, FrameTransform, VolumeRes3f, CompIdx = (uint32)PreviewAttribute, Ext = Extinction, Mip = MipLevel, VoxelSizeFactor = VoxelSize, bPivotAtCorner = (bool)bLocalOriginAtCorner](FRHICommandList& RHICmdList)
+		{
+			SVTSceneProxy->GlobalTransform = GlobalTransform;
+			SVTSceneProxy->FrameTransform = FrameTransform;
+			SVTSceneProxy->VolumeResolution = VolumeRes3f;
+			SVTSceneProxy->MipLevel = Mip;
+			SVTSceneProxy->ComponentToVisualize = CompIdx;
+			SVTSceneProxy->Extinction = Ext;
+			SVTSceneProxy->VoxelSizeFactor = VoxelSizeFactor;
+			SVTSceneProxy->bPivotAtCorner = bPivotAtCorner;
+		});
 	}
 }
 
diff --git a/Engine/Source/Runtime/Renderer/Private/SparseVolumeTexture/SparseVolumeTextureViewerComponent.h b/Engine/Source/Runtime/Renderer/Private/SparseVolumeTexture/SparseVolumeTextureViewerComponent.h
--- a/Engine/Source/Runtime/Renderer/Private/SparseVolumeTexture/SparseVolumeTextureViewerComponent.h
+++ b/Engine/Source/Runtime/Renderer/Private/SparseVolumeTexture/SparseVolumeTextureViewerComponent.h
@@ -50,6 +50,15 @@ class USparseVolumeTextureViewerComponent : public UPrimitiveComponent
 	UPROPERTY(EditAnywhere, BlueprintReadOnly, Category = "Asset Preview")
 	uint32 bBlockingStreamingRequests : 1;
 
+	UPROPERTY(EditAnywhere, BlueprintReadOnly, Category = "Asset Preview")
+	uint32 bApplyPerFrameTransforms : 1;
+
+	UPROPERTY(EditAnywhere, BlueprintReadOnly, Category = "Asset Preview")
+	uint32 bLocalOriginAtCorner : 1;
+
+	UPROPERTY(EditAnywhere, BlueprintReadWrite, Category = "Asset Preview")
+	float VoxelSize = 1.0f;
+
 	UPROPERTY(EditAnywhere, BlueprintReadWrite, Category = "Asset Preview", meta = (UIMin = 0.0, UIMax = 1.0, ClampMin = 0.0, ClampMax = 1.0, EditCondition = "!bAnimate"))
 	float AnimationFrame;
 
diff --git a/Engine/Source/Runtime/Renderer/Private/SparseVolumeTexture/SparseVolumeTextureViewerRendering.cpp b/Engine/Source/Runtime/Renderer/Private/SparseVolumeTexture/SparseVolumeTextureViewerRendering.cpp
--- a/Engine/Source/Runtime/Renderer/Private/SparseVolumeTexture/SparseVolumeTextureViewerRendering.cpp
+++ b/Engine/Source/Runtime/Renderer/Private/SparseVolumeTexture/SparseVolumeTextureViewerRendering.cpp
@@ -67,9 +67,9 @@ class FVisualizeSparseVolumeTexturePS : public FGlobalShader
 		SHADER_PARAMETER(FVector4f, WorldToLocal0)
 		SHADER_PARAMETER(FVector4f, WorldToLocal1)
 		SHADER_PARAMETER(FVector4f, WorldToLocal2)
-		SHADER_PARAMETER(FVector3f, WorldToLocalNoScale0)
-		SHADER_PARAMETER(FVector3f, WorldToLocalNoScale1)
-		SHADER_PARAMETER(FVector3f, WorldToLocalNoScale2)
+		SHADER_PARAMETER(FVector3f, WorldToLocalRotation0)
+		SHADER_PARAMETER(FVector3f, WorldToLocalRotation1)
+		SHADER_PARAMETER(FVector3f, WorldToLocalRotation2)
 		SHADER_PARAMETER(uint32, ComponentToVisualize)
 		SHADER_PARAMETER(float, Extinction)
 	END_SHADER_PARAMETER_STRUCT()
@@ -118,8 +118,28 @@ void AddSparseVolumeTextureViewerRenderPass(FRDGBuilder& GraphBuilder, FSceneRen
 
 		for (auto& SVTSceneProxy : Scene->SparseVolumeTextureViewers)
 		{
-			const FMatrix44f& WorldToLocal = SVTSceneProxy->WorldToLocal;
-			const FMatrix44f& WorldToLocalNoScale = SVTSceneProxy->WorldToLocalNoScale;
+			const FVector& PreViewTranslation = View.ViewMatrices.GetPreViewTranslation();
+
+			FTransform GlobalTransform = SVTSceneProxy->GlobalTransform;
+			GlobalTransform.AddToTranslation(PreViewTranslation); // Move into translated world space
+
+			FTransform FrameTransform = SVTSceneProxy->FrameTransform;
+			FrameTransform.MultiplyScale3D(FVector(SVTSceneProxy->VoxelSizeFactor));
+			FrameTransform.ScaleTranslation(SVTSceneProxy->VoxelSizeFactor);
+
+			const FTransform Transform = FrameTransform * GlobalTransform;
+			const FMatrix InvTransformMat = Transform.ToMatrixWithScale().Inverse();
+			const FQuat InvRotation = Transform.GetRotation().Inverse();
+
+			const FVector RcpVolumeRes = FVector(SVTSceneProxy->VolumeResolution).Reciprocal();
+			FMatrix UnitSpaceMat = FMatrix::Identity;
+			UnitSpaceMat.SetColumn(0, FVector(RcpVolumeRes.X * 2.0, 0.0, 0.0));
+			UnitSpaceMat.SetColumn(1, FVector(0.0, RcpVolumeRes.Y * 2.0, 0.0));
+			UnitSpaceMat.SetColumn(2, FVector(0.0, 0.0, RcpVolumeRes.Z * 2.0));
+			UnitSpaceMat.SetOrigin(FVector(SVTSceneProxy->bPivotAtCorner ? -1.0 : 0.0));
+
+			const FMatrix44f WorldToLocal = FMatrix44f(InvTransformMat * UnitSpaceMat);
+			const FMatrix44f WorldToLocalRot = FMatrix44f(InvRotation.ToMatrix());
 			
 			FGlobalShaderMap* GlobalShaderMap = GetGlobalShaderMap(GMaxRHIFeatureLevel);
 			FVisualizeSparseVolumeTextureVS::FPermutationDomain VsPermutationVector;
@@ -141,9 +161,9 @@ void AddSparseVolumeTextureViewerRenderPass(FRDGBuilder& GraphBuilder, FSceneRen
 			PsPassParameters->WorldToLocal0 = FVector4f(WorldToLocal.M[0][0], WorldToLocal.M[1][0], WorldToLocal.M[2][0], WorldToLocal.M[3][0]);
 			PsPassParameters->WorldToLocal1 = FVector4f(WorldToLocal.M[0][1], WorldToLocal.M[1][1], WorldToLocal.M[2][1], WorldToLocal.M[3][1]);
 			PsPassParameters->WorldToLocal2 = FVector4f(WorldToLocal.M[0][2], WorldToLocal.M[1][2], WorldToLocal.M[2][2], WorldToLocal.M[3][2]);
-			PsPassParameters->WorldToLocalNoScale0 = FVector3f(WorldToLocalNoScale.M[0][0], WorldToLocalNoScale.M[1][0], WorldToLocalNoScale.M[2][0]);
-			PsPassParameters->WorldToLocalNoScale1 = FVector3f(WorldToLocalNoScale.M[0][1], WorldToLocalNoScale.M[1][1], WorldToLocalNoScale.M[2][1]);
-			PsPassParameters->WorldToLocalNoScale2 = FVector3f(WorldToLocalNoScale.M[0][2], WorldToLocalNoScale.M[1][2], WorldToLocalNoScale.M[2][2]);
+			PsPassParameters->WorldToLocalRotation0 = FVector3f(WorldToLocalRot.M[0][0], WorldToLocalRot.M[1][0], WorldToLocalRot.M[2][0]);
+			PsPassParameters->WorldToLocalRotation1 = FVector3f(WorldToLocalRot.M[0][1], WorldToLocalRot.M[1][1], WorldToLocalRot.M[2][1]);
+			PsPassParameters->WorldToLocalRotation2 = FVector3f(WorldToLocalRot.M[0][2], WorldToLocalRot.M[1][2], WorldToLocalRot.M[2][2]);
 			PsPassParameters->ComponentToVisualize = SVTSceneProxy->ComponentToVisualize;
 			PsPassParameters->Extinction = SVTSceneProxy->Extinction;
 			
diff --git a/Engine/Source/Runtime/Renderer/Private/SparseVolumeTexture/SparseVolumeTextureViewerSceneProxy.h b/Engine/Source/Runtime/Renderer/Private/SparseVolumeTexture/SparseVolumeTextureViewerSceneProxy.h
--- a/Engine/Source/Runtime/Renderer/Private/SparseVolumeTexture/SparseVolumeTextureViewerSceneProxy.h
+++ b/Engine/Source/Runtime/Renderer/Private/SparseVolumeTexture/SparseVolumeTextureViewerSceneProxy.h
@@ -25,12 +25,14 @@ public:
 	virtual ~FSparseVolumeTextureViewerSceneProxy() = default;
 
 	const UE::SVT::FTextureRenderResources* TextureRenderResources;
-	FMatrix44f WorldToLocal;
-	FMatrix44f WorldToLocalNoScale;
+	FTransform GlobalTransform;
+	FTransform FrameTransform;
 	FVector3f VolumeResolution;
 	int32 MipLevel;
 	uint32 ComponentToVisualize;
 	float Extinction;
+	float VoxelSizeFactor;
+	bool bPivotAtCorner;
 
 protected:
 
