Fixed OpenVDB import window documentation button not linking to any valid website.

#rnx

diff --git a/Engine/Source/Editor/SparseVolumeTexture/Private/OpenVDBImportWindow.cpp b/Engine/Source/Editor/SparseVolumeTexture/Private/OpenVDBImportWindow.cpp
--- a/Engine/Source/Editor/SparseVolumeTexture/Private/OpenVDBImportWindow.cpp
+++ b/Engine/Source/Editor/SparseVolumeTexture/Private/OpenVDBImportWindow.cpp
@@ -234,7 +234,7 @@ void SOpenVDBImportWindow::Construct(const FArguments& InArgs)
 			.VAlign(VAlign_Center)
 			.AutoWidth()
 			[
-				IDocumentation::Get()->CreateAnchor(FString("Engine/Content/OpenVDB/ImportWindow"))
+				IDocumentation::Get()->CreateAnchor(FString("sparse-volume-textures-in-unreal-engine"))
 			]
 			+ SHorizontalBox::Slot()
 			.HAlign(HAlign_Right)
