SparseVolumeTexture: Fix one last place still using PLATFORM_WINDOWS || PLATFORM_LINUX instead of OPENVDB_AVAILABLE

diff --git a/Engine/Source/Editor/SparseVolumeTexture/Private/SparseVolumeTextureModule.cpp b/Engine/Source/Editor/SparseVolumeTexture/Private/SparseVolumeTextureModule.cpp
--- a/Engine/Source/Editor/SparseVolumeTexture/Private/SparseVolumeTextureModule.cpp
+++ b/Engine/Source/Editor/SparseVolumeTexture/Private/SparseVolumeTextureModule.cpp
@@ -19,7 +19,7 @@ static void RegisterOpenVDBGrid()
 
 void FSparseVolumeTextureModule::StartupModule()
 {
-#if PLATFORM_WINDOWS || PLATFORM_LINUX
+#if OPENVDB_AVAILABLE
 	// Global registration of  the vdb types.
 	openvdb::initialize();
 
