[SparseVolumeTexture] Disable AutoRTFM instrumentation.

AutoRTFM instrumentation cannot be enabled with exceptions.

diff --git a/Engine/Source/Editor/SparseVolumeTexture/SparseVolumeTexture.Build.cs b/Engine/Source/Editor/SparseVolumeTexture/SparseVolumeTexture.Build.cs
--- a/Engine/Source/Editor/SparseVolumeTexture/SparseVolumeTexture.Build.cs
+++ b/Engine/Source/Editor/SparseVolumeTexture/SparseVolumeTexture.Build.cs
@@ -54,6 +54,7 @@ public class SparseVolumeTexture : ModuleRules
 		if (Target.IsInPlatformGroup(UnrealPlatformGroup.Windows))
 		{
 			bUseRTTI = true;
+			bDisableAutoRTFMInstrumentation = true; // AutoRTFM cannot be used with exceptions
 			bEnableExceptions = true;
 			PublicDefinitions.Add("OPENVDB_AVAILABLE=1");
 
