Make SparseVolumeTexture work for editor + linux by enabling exceptions for the module on linux too.

diff --git a/Engine/Source/Editor/SparseVolumeTexture/SparseVolumeTexture.Build.cs b/Engine/Source/Editor/SparseVolumeTexture/SparseVolumeTexture.Build.cs
--- a/Engine/Source/Editor/SparseVolumeTexture/SparseVolumeTexture.Build.cs
+++ b/Engine/Source/Editor/SparseVolumeTexture/SparseVolumeTexture.Build.cs
@@ -50,11 +50,12 @@ public class SparseVolumeTexture : ModuleRules
 			);
 		}
 
+		bDisableAutoRTFMInstrumentation = true;
+
 		// Specific to OpenVDB support
 		if (Target.IsInPlatformGroup(UnrealPlatformGroup.Windows))
 		{
 			bUseRTTI = true;
-			bDisableAutoRTFMInstrumentation = true; // AutoRTFM cannot be used with exceptions
 			bEnableExceptions = true;
 			PublicDefinitions.Add("OPENVDB_AVAILABLE=1");
 
@@ -69,7 +70,7 @@ public class SparseVolumeTexture : ModuleRules
 		else if (Target.Platform == UnrealTargetPlatform.Linux)
 		{
 			bUseRTTI = false;
-			bEnableExceptions = false;
+			bEnableExceptions = true;
 			PublicDefinitions.Add("OPENVDB_AVAILABLE=1");
 
 			AddEngineThirdPartyPrivateStaticDependencies(Target,
