[Sparse Volume Texture] Fixed dangling pointer bug.

diff --git a/Engine/Source/Editor/SparseVolumeTexture/Private/SparseVolumeTextureOpenVDBUtility.cpp b/Engine/Source/Editor/SparseVolumeTexture/Private/SparseVolumeTextureOpenVDBUtility.cpp
--- a/Engine/Source/Editor/SparseVolumeTexture/Private/SparseVolumeTextureOpenVDBUtility.cpp
+++ b/Engine/Source/Editor/SparseVolumeTexture/Private/SparseVolumeTextureOpenVDBUtility.cpp
@@ -363,7 +363,7 @@ public:
 
 		FIntVector3 SmallestAABBMin = FIntVector3(INT32_MAX);
 		FIntVector3 LargestAABBMax = FIntVector3(INT32_MIN);
-		const FTransform* LastGridTransformPtr = nullptr;
+		TOptional<FTransform> PreviousGridTransform;
 
 		// Compute per source grid data of up to 4 different grids (one per component)
 		UniqueGridAdapters.SetNum((int32)NumSourceGrids);
@@ -405,11 +405,11 @@ public:
 				LargestAABBMax.Y = FMath::Max(LargestAABBMax.Y, GridInfo.VolumeActiveAABBMax.Y);
 				LargestAABBMax.Z = FMath::Max(LargestAABBMax.Z, GridInfo.VolumeActiveAABBMax.Z);
 
-				if (LastGridTransformPtr && !LastGridTransformPtr->Equals(GridInfo.Transform))
+				if (PreviousGridTransform && !PreviousGridTransform->Equals(GridInfo.Transform))
 				{
 					UE_LOG(LogSparseVolumeTextureOpenVDBUtility, Warning, TEXT("Frame has multiple grids with different transforms in the same frame! Data will likely not be imported/displayed correctly!"));
 				}
-				LastGridTransformPtr = &GridInfo.Transform;
+				PreviousGridTransform = GridInfo.Transform;
 
 				SVTCreateInfo.FallbackValues[AttributesIdx][CompIdx] = UniqueGridAdapters[SourceGridIndex]->GetBackgroundValue(SourceComponentIndex);
 
@@ -425,7 +425,7 @@ public:
 		SVTCreateInfo.VirtualVolumeAABBMin = bEmptyBounds ? FIntVector3(INT32_MAX) : (SmallestAABBMin - VolumeBoundsMin);
 		SVTCreateInfo.VirtualVolumeAABBMax = bEmptyBounds ? FIntVector3(INT32_MIN) : (LargestAABBMax - VolumeBoundsMin);
 		
-		FrameTransform = LastGridTransformPtr ? *LastGridTransformPtr : FTransform::Identity;
+		FrameTransform = PreviousGridTransform ? *PreviousGridTransform : FTransform::Identity;
 		FrameTransform.AddToTranslation(FVector(VolumeBoundsMin) * FrameTransform.GetScale3D()); // We made the volume relative to VolumeBoundsMin, so we need to undo this translation when using the frame transform
 
 		return true;
