SparseVolumeTexture, EngineTest: re-encode explosion.vdb with newer file format

OpenVDB 13.0.0 removed support for some very old VDB file format versions:
[REDACTED_PROFILE_URL]/openvdb/pull/2087

The test file explosion.vdb was the only file in the EngineTest project
with a file version (220) older than the cutoff (222).

Re-encoding the test file uncovered a minor issue in the order of operations
with optimized density assignment where the order in which grid components
were considered could cause an ensure to fire erroneously.

diff --git a/Engine/Source/Editor/SparseVolumeTexture/Private/SparseVolumeTextureFactory.cpp b/Engine/Source/Editor/SparseVolumeTexture/Private/SparseVolumeTextureFactory.cpp
--- a/Engine/Source/Editor/SparseVolumeTexture/Private/SparseVolumeTextureFactory.cpp
+++ b/Engine/Source/Editor/SparseVolumeTexture/Private/SparseVolumeTextureFactory.cpp
@@ -52,7 +52,8 @@ static void ComputeDefaultOpenVDBGridAssignment(const TArray<TSharedPtr<FOpenVDB
 	// Assign the components of the input grids to the components of the output SVT.
 	
 	const TSharedPtr<FOpenVDBGridComponentInfo>* DensityComponentInfoPtr = GridComponentInfo.FindByPredicate([](const TSharedPtr<FOpenVDBGridComponentInfo>& GridComponent) { return GridComponent->Name == TEXT("density"); });
-	const int32 NumNonDensityComponentInfos = GridComponentInfo.Num() - 1 - (DensityComponentInfoPtr ? 1 : 0); // -1 because there is always a <None> element in the list
+	const TSharedPtr<FOpenVDBGridComponentInfo>* NoneComponentInfoPtr = GridComponentInfo.FindByPredicate([](const TSharedPtr<FOpenVDBGridComponentInfo>& GridComponent) { return GridComponent->Name == TEXT("<None>"); });
+	const int32 NumNonDensityComponentInfos = GridComponentInfo.Num() - (NoneComponentInfoPtr ? 1 : 0) - (DensityComponentInfoPtr ? 1 : 0);
 	
 	// Optimized density assignment: density as 8bit unorm in Attributes A and all other components in Attributes B as 16bit float. This only works if there is a maximum of 4 non-density components.
 	// We also don't use this assignment if there are 3 non-density components as these will be padded to a 4 component format anyways, so we might as well put all 4 components into a single texture.
@@ -69,7 +70,7 @@ static void ComputeDefaultOpenVDBGridAssignment(const TArray<TSharedPtr<FOpenVDB
 		uint32 DstComponentIdx = 0;
 		for (const TSharedPtr<FOpenVDBGridComponentInfo>& GridComponent : GridComponentInfo)
 		{
-			if (!ensure(DstComponentIdx <= 3) || GridComponent->Index == INDEX_NONE || &GridComponent == DensityComponentInfoPtr)
+			if (GridComponent->Index == INDEX_NONE || &GridComponent == DensityComponentInfoPtr || !ensure(DstComponentIdx <= 3))
 			{
 				continue;
 			}
