USD: Fix stage transform not being baked into the volume.

#Synced-CL 50957406

diff --git a/Engine/Plugins/Importers/USDImporter/Source/USDSchemas/Private/USDVolVolumeTranslator.cpp b/Engine/Plugins/Importers/USDImporter/Source/USDSchemas/Private/USDVolVolumeTranslator.cpp
--- a/Engine/Plugins/Importers/USDImporter/Source/USDSchemas/Private/USDVolVolumeTranslator.cpp
+++ b/Engine/Plugins/Importers/USDImporter/Source/USDSchemas/Private/USDVolVolumeTranslator.cpp
@@ -598,6 +598,12 @@ void FUsdVolVolumeTranslator::CreateAssets()
 
 	FString VolumePrimHashPrefix = UsdUtils::GetAssetHashPrefix(GetPrim(), Context->bShareAssetsForIdenticalPrims);
 
+	// Bake our stage coordinate system transform into the volume, just like we do for geometry.
+	// We do this here because OpenVDB doesn't really specify a fixed coordinate system, and so we interpret it as just meaning
+	// "units" in general. This means we must go through a metersPerUnit / upAxis conversion to get it into UE's world space
+	const FUsdStageInfo StageInfo{Stage};
+	FTransform StageInfoTransform = UsdUtils::GetUsdToUESpaceTransform(StageInfo);
+
 	// Collect info from requested files
 	const TMap<FString, UsdUtils::FVolumePrimInfo> FilePathHashToVolumeInfo = UsdUtils::GetVolumeInfoByFilePathHash(VolumePrim);
 
@@ -659,7 +665,7 @@ void FUsdVolVolumeTranslator::CreateAssets()
 			SparseVolumeTextureInfo.PrefixedAssetHash,
 			DesiredName,
 			Context->ObjectFlags,
-			[&VDBFilePath, &ImportOptions](UPackage* Outer, FName SanitizedName, EObjectFlags FlagsToUse)
+			[&VDBFilePath, &ImportOptions, &StageInfoTransform](UPackage* Outer, FName SanitizedName, EObjectFlags FlagsToUse)
 			{
 				TStrongObjectPtr<USparseVolumeTextureFactory> SparseVolumeTextureFactory{NewObject<USparseVolumeTextureFactory>()};
 
@@ -677,15 +683,15 @@ void FUsdVolVolumeTranslator::CreateAssets()
 				// the asset cache will do that anyway
 				const TCHAR* Parms = nullptr;
 				bool bOperationCanceled = false;
-				return Cast<USparseVolumeTexture>(SparseVolumeTextureFactory->FactoryCreateFile(
+				return Cast<USparseVolumeTexture>(SparseVolumeTextureFactory->ImportWithTransform(
 					USparseVolumeTexture::StaticClass(),
 					Outer,
 					SanitizedName,
 					FlagsToUse,
 					VDBFilePath,
 					Parms,
-					GWarn,
-					bOperationCanceled
+					bOperationCanceled,
+					StageInfoTransform
 				));
 			},
 			&bCreatedNew
diff --git a/Engine/Source/Editor/SparseVolumeTexture/Private/SparseVolumeTextureFactory.cpp b/Engine/Source/Editor/SparseVolumeTexture/Private/SparseVolumeTextureFactory.cpp
--- a/Engine/Source/Editor/SparseVolumeTexture/Private/SparseVolumeTextureFactory.cpp
+++ b/Engine/Source/Editor/SparseVolumeTexture/Private/SparseVolumeTextureFactory.cpp
@@ -419,7 +419,12 @@ EReimportResult::Type USparseVolumeTextureFactory::Reimport(UObject* Obj)
 #endif // OPENVDB_AVAILABLE
 }
 
-UObject* USparseVolumeTextureFactory::ImportInternal(UClass* InClass, UObject* InParent, FName InName, EObjectFlags Flags, const FString& Filename, const TCHAR* Parms, bool& bOutOperationCanceled, bool bIsReimport)
+UObject* USparseVolumeTextureFactory::ImportWithTransform(UClass* InClass, UObject* InParent, FName InName, EObjectFlags Flags, const FString& Filename, const TCHAR* Parms, bool& bOutOperationCanceled, const FTransform& AdditionalFrameTransform)
+{
+	return ImportInternal(InClass, InParent, InName, Flags, Filename, Parms, bOutOperationCanceled, false /*bIsReimport*/, AdditionalFrameTransform);
+}
+
+UObject* USparseVolumeTextureFactory::ImportInternal(UClass* InClass, UObject* InParent, FName InName, EObjectFlags Flags, const FString& Filename, const TCHAR* Parms, bool& bOutOperationCanceled, bool bIsReimport, const FTransform& AdditionalFrameTransform)
 {
 #if OPENVDB_AVAILABLE
 
@@ -518,6 +523,8 @@ UObject* USparseVolumeTextureFactory::ImportInternal(UClass* InClass, UObject* I
 			return nullptr;
 		}
 
+		FrameTransform = FrameTransform * AdditionalFrameTransform;
+
 		UStaticSparseVolumeTexture* StaticSVTexture = NewObject<UStaticSparseVolumeTexture>(InParent, UStaticSparseVolumeTexture::StaticClass(), InName, Flags);
 		const bool bInitSuccess = StaticSVTexture->Initialize(MakeArrayView(&TextureData, 1), MakeArrayView(&FrameTransform, 1));
 		if (!bInitSuccess)
@@ -649,7 +656,7 @@ UObject* USparseVolumeTextureFactory::ImportInternal(UClass* InClass, UObject* I
 
 			AsyncTask(ENamedThreads::AnyNormalThreadNormalTask,
 				[FrameIdx, NumFrames, &PreviewData, &ImportOptions, &UncookedFramesData, &FrameTransforms,
-				AllTasksFinishedEvent, &bErrored, &bCanceled, &FinishedTasksCounter, &ProcessedFramesCounter, &VolumeBoundsMin]()
+				AllTasksFinishedEvent, &bErrored, &bCanceled, &FinishedTasksCounter, &ProcessedFramesCounter, &VolumeBoundsMin, &AdditionalFrameTransform]()
 				{
 					// Ensure the FinishedTasksCounter will be incremented in all cases
 					FScopedIncrementer Incremeter(FinishedTasksCounter, NumFrames, AllTasksFinishedEvent);
@@ -684,7 +691,7 @@ UObject* USparseVolumeTextureFactory::ImportInternal(UClass* InClass, UObject* I
 					}
 
 					UncookedFramesData[FrameIdx] = MoveTemp(TextureData);
-					FrameTransforms[FrameIdx] = FrameTransform;
+					FrameTransforms[FrameIdx] = FrameTransform * AdditionalFrameTransform;
 
 					// Increment ProcessedFramesCounter
 					ProcessedFramesCounter.fetch_add(1);
diff --git a/Engine/Source/Editor/SparseVolumeTexture/Public/SparseVolumeTextureFactory.h b/Engine/Source/Editor/SparseVolumeTexture/Public/SparseVolumeTextureFactory.h
--- a/Engine/Source/Editor/SparseVolumeTexture/Public/SparseVolumeTextureFactory.h
+++ b/Engine/Source/Editor/SparseVolumeTexture/Public/SparseVolumeTextureFactory.h
@@ -52,10 +52,13 @@ public:
 	virtual EReimportResult::Type Reimport(UObject* Obj) override;
 	//~ End FReimportHandler Interface
 
+	SPARSEVOLUMETEXTURE_API UObject* ImportWithTransform(UClass* InClass, UObject* InParent, FName InName, EObjectFlags Flags, const FString& Filename,
+		const TCHAR* Parms, bool& bOutOperationCanceled, const FTransform& AdditionalFrameTransform);
+
 protected:
 
 private:
-	UObject* ImportInternal(UClass* InClass, UObject* InParent, FName InName, EObjectFlags Flags, const FString& Filename, const TCHAR* Parms, bool& bOutOperationCanceled, bool bIsReimport);
+	UObject* ImportInternal(UClass* InClass, UObject* InParent, FName InName, EObjectFlags Flags, const FString& Filename, const TCHAR* Parms, bool& bOutOperationCanceled, bool bIsReimport, const FTransform& AdditionalFrameTransform = FTransform::Identity);
 };
 
 bool SPARSEVOLUMETEXTURE_API LoadOpenVDBPreviewData(const FString& Filename, FOpenVDBPreviewData* OutPreviewData);
